Content Browser: Add filters for file types such as Python

diff --git a/Engine/Plugins/Editor/ContentBrowser/ContentBrowserFileDataSource/Source/ContentBrowserFileDataSource/Private/ContentBrowserFileDataSource.cpp b/Engine/Plugins/Editor/ContentBrowser/ContentBrowserFileDataSource/Source/ContentBrowserFileDataSource/Private/ContentBrowserFileDataSource.cpp
--- a/Engine/Plugins/Editor/ContentBrowser/ContentBrowserFileDataSource/Source/ContentBrowserFileDataSource/Private/ContentBrowserFileDataSource.cpp
+++ b/Engine/Plugins/Editor/ContentBrowser/ContentBrowserFileDataSource/Source/ContentBrowserFileDataSource/Private/ContentBrowserFileDataSource.cpp
@@ -16,6 +16,8 @@
 #include "DirectoryWatcherModule.h"
 #include "ContentBrowserDataSubsystem.h"
 #include "ContentBrowserDataUtils.h"
+#include "AssetTypeCategories.h"
+#include "AssetTypeActions_Base.h"
 
 #define LOCTEXT_NAMESPACE "ContentBrowserFileDataSource"
 
@@ -34,6 +36,57 @@ namespace ContentBrowserFileDataSource
 
 		return InPath;
 	}
+
+	FString GetFilterNamePrefix()
+	{
+		return TEXT("ContentBrowserDataFilterFile");
+	}
+}
+
+/** Asset type actions for files filter */
+class FAssetTypeActions_FileDataSource : public FAssetTypeActions_Base
+{
+public:
+	// Begin IAssetTypeActions interface
+	virtual uint32 GetCategories() override { return EAssetTypeCategories::Misc; }
+	virtual bool IsImportedAsset() const override { return false; }
+
+	virtual FText GetName() const override;
+	virtual FText GetAssetDescription(const FAssetData& AssetData) const override;
+	virtual UClass* GetSupportedClass() const override;
+	virtual FColor GetTypeColor() const override;
+	virtual FName GetFilterName() const override;
+	// End IAssetTypeActions interface
+
+	FName FilterName;
+	FText Name;
+	FText Description;
+	FColor TypeColor;
+};
+
+FText FAssetTypeActions_FileDataSource::GetAssetDescription(const struct FAssetData& AssetData) const
+{
+	return Description;
+}
+
+UClass* FAssetTypeActions_FileDataSource::GetSupportedClass() const
+{
+	return nullptr;
+}
+
+FText FAssetTypeActions_FileDataSource::GetName() const
+{
+	return Name;
+}
+
+FColor FAssetTypeActions_FileDataSource::GetTypeColor() const
+{
+	return TypeColor;
+}
+
+FName FAssetTypeActions_FileDataSource::GetFilterName() const
+{
+	return FilterName;
 }
 
 class FContentBrowserFileDataDiscovery : public FRunnable
@@ -311,10 +364,33 @@ void UContentBrowserFileDataSource::Initialize(const ContentBrowserFileData::FFi
 			}));
 		}
 	}
+
+	// Register asset type actions
+	IAssetTools& AssetTools = FModuleManager::LoadModuleChecked<FAssetToolsModule>("AssetTools").Get();
+	Config.EnumerateFileActions([this, &AssetTools](TSharedRef<const ContentBrowserFileData::FFileActions> InFileActions)
+	{
+		TSharedRef<FAssetTypeActions_FileDataSource> AssetTypeAction = MakeShared<FAssetTypeActions_FileDataSource>();
+		AssetTypeAction->Name = InFileActions->TypeShortDescription;
+		AssetTypeAction->Description = InFileActions->TypeFullDescription;
+		AssetTypeAction->TypeColor = InFileActions->TypeColor.ToFColor(true);
+		AssetTypeAction->FilterName = *(ContentBrowserFileDataSource::GetFilterNamePrefix() + InFileActions->TypeName.ToString());
+		AssetTools.RegisterAssetTypeActions(AssetTypeAction);
+		RegisteredAssetTypeActions.Add(AssetTypeAction);
+		return true;
+	});
 }
 
 void UContentBrowserFileDataSource::Shutdown()
 {
+	// Unregister asset type actions
+	if (FAssetToolsModule* AssetToolsModule = FModuleManager::GetModulePtr<FAssetToolsModule>("AssetTools"))
+	{
+		for (const TSharedRef<IAssetTypeActions>& Action : RegisteredAssetTypeActions)
+		{
+			AssetToolsModule->Get().UnregisterAssetTypeActions(Action);
+		}
+	}
+
 	if (BackgroundDiscovery)
 	{
 		BackgroundDiscovery->EnsureCompletion();
@@ -488,16 +564,55 @@ void UContentBrowserFileDataSource::CompileFilter(const FName InPath, const FCon
 		return;
 	}
 
-	// If we have any inclusive object, path, or class filters, then skip files
-	// TODO: Add a way to filter on file types
+	TArray<FString> FileExtensionsToInclude;
+	if (ClassFilter)
+	{
+		if (ClassFilter->ClassNamesToInclude.Num() > 0)
+		{
+			TArray<FName, TInlineAllocator<2>> ClassNamesToLookFor;
+			TArray<FString, TInlineAllocator<2>> ClassFileExtensions;
+
+			Config.EnumerateFileActions([&ClassNamesToLookFor, &ClassFileExtensions](TSharedRef<const ContentBrowserFileData::FFileActions> InFileActions)
+			{
+				TStringBuilder<64> ClassNameToLookFor;
+				ClassNameToLookFor.Append(ContentBrowserFileDataSource::GetFilterNamePrefix());
+				FNameBuilder TypeNameBuilder(InFileActions->TypeName);
+				ClassNameToLookFor.Append(TypeNameBuilder);
+				ClassNamesToLookFor.Add(ClassNameToLookFor.ToString());
+
+				ClassFileExtensions.Add(InFileActions->TypeExtension);
+				return true;
+			});
+
+			// Determine if any are one of this instance's file data source
+			for (const FName ClassName : ClassFilter->ClassNamesToInclude)
+			{
+				int32 FoundIndex = INDEX_NONE;
+				if (ClassNamesToLookFor.Find(ClassName, FoundIndex))
+				{
+					FileExtensionsToInclude.Add(ClassFileExtensions[FoundIndex]);
+				}
+			}
+		}
+	}
+
+	// If we have any inclusive object, path, then skip files
 	if ((ObjectFilter && (ObjectFilter->ObjectNamesToInclude.Num() > 0 || ObjectFilter->TagsAndValuesToInclude.Num() > 0)) ||
-		(PackageFilter && (PackageFilter->PackageNamesToInclude.Num() > 0 || PackageFilter->PackagePathsToInclude.Num() > 0)) ||
-		(ClassFilter && (ClassFilter->ClassNamesToInclude.Num() > 0))
+		(PackageFilter && (PackageFilter->PackageNamesToInclude.Num() > 0 || PackageFilter->PackagePathsToInclude.Num() > 0))
 		)
 	{
 		return;
 	}
 
+	// When a class filter is active, only show files if a file filter is also active
+	if (ClassFilter && (ClassFilter->ClassNamesToInclude.Num() > 0))
+	{
+		if (FileExtensionsToInclude.Num() == 0)
+		{
+			return;
+		}
+	}
+
 	// TODO: Support adding loose files to collections via their internal mount paths - for now just bail
 	if (CollectionFilter && CollectionFilter->SelectedCollections.Num() > 0)
 	{
@@ -509,6 +624,7 @@ void UContentBrowserFileDataSource::CompileFilter(const FName InPath, const FCon
 	FileDataFilter.VirtualPath = InPath.ToString();
 	FileDataFilter.bRecursivePaths = InFilter.bRecursivePaths;
 	FileDataFilter.ItemAttributeFilter = InFilter.ItemAttributeFilter;
+	FileDataFilter.FileExtensionsToInclude.Append(FileExtensionsToInclude);
 	if (PackageFilter && PackageFilter->PathBlacklist && PackageFilter->PathBlacklist->HasFiltering())
 	{
 		FileDataFilter.Blacklist = PackageFilter->PathBlacklist;
@@ -658,7 +774,7 @@ void UContentBrowserFileDataSource::EnumerateItemsMatchingFilter(const FContentB
 					}
 					else if (DiscoveredChildItem->Type == FDiscoveredItem::EType::File)
 					{
-						if (bIncludeFiles)
+						if (bIncludeFiles && PassesFileTypeFilter(DiscoveredChildItem->DiskPath))
 						{
 							if (!Callback(DataSource->CreateFileItem(ChildItemName, DiscoveredChildItem->DiskPath)))
 							{
@@ -672,6 +788,26 @@ void UContentBrowserFileDataSource::EnumerateItemsMatchingFilter(const FContentB
 			return true;
 		};
 
+		bool PassesFileTypeFilter(const FStringView DiskPath)
+		{
+			if (FileDataFilter.FileExtensionsToInclude.Num() == 0)
+			{
+				return true;
+			}
+
+			int32 LastDot = INDEX_NONE;
+			if (DiskPath.FindLastChar(TEXT('.'), LastDot))
+			{
+				const FStringView FileExtension = DiskPath.Right(DiskPath.Len() - LastDot - 1);
+				if (FileDataFilter.FileExtensionsToInclude.Contains(FileExtension))
+				{
+					return true;
+				}
+			}
+
+			return false;
+		}
+
 		bool IncludeFolders() const
 		{
 			return bIncludeFolders;
diff --git a/Engine/Plugins/Editor/ContentBrowser/ContentBrowserFileDataSource/Source/ContentBrowserFileDataSource/Public/ContentBrowserFileDataSource.h b/Engine/Plugins/Editor/ContentBrowser/ContentBrowserFileDataSource/Source/ContentBrowserFileDataSource/Public/ContentBrowserFileDataSource.h
--- a/Engine/Plugins/Editor/ContentBrowser/ContentBrowserFileDataSource/Source/ContentBrowserFileDataSource/Public/ContentBrowserFileDataSource.h
+++ b/Engine/Plugins/Editor/ContentBrowser/ContentBrowserFileDataSource/Source/ContentBrowserFileDataSource/Public/ContentBrowserFileDataSource.h
@@ -10,6 +10,7 @@
 
 class UToolMenu;
 class FContentBrowserFileDataDiscovery;
+class IAssetTypeActions;
 struct FFileChangeData;
 
 USTRUCT()
@@ -23,6 +24,7 @@ public:
 	bool bRecursivePaths = false;
 	EContentBrowserItemAttributeFilter ItemAttributeFilter = EContentBrowserItemAttributeFilter::IncludeNone;
 	TSharedPtr<FBlacklistPaths> Blacklist;
+	TArray<FString> FileExtensionsToInclude;
 };
 
 UCLASS()
@@ -177,6 +179,8 @@ protected:
 
 	ContentBrowserFileData::FFileConfigData Config;
 
+	TArray<TSharedRef<IAssetTypeActions>> RegisteredAssetTypeActions;
+
 	TSortedMap<FName, FFileMount, FDefaultAllocator, FNameFastLess> RegisteredFileMounts;
 	TMap<FName, TArray<FName>> RegisteredFileMountRoots;
 
diff --git a/Engine/Source/Developer/AssetTools/Private/AssetTools.cpp b/Engine/Source/Developer/AssetTools/Private/AssetTools.cpp
--- a/Engine/Source/Developer/AssetTools/Private/AssetTools.cpp
+++ b/Engine/Source/Developer/AssetTools/Private/AssetTools.cpp
@@ -315,8 +315,17 @@ UAssetToolsImpl::UAssetToolsImpl(const FObjectInitializer& ObjectInitializer)
 
 void UAssetToolsImpl::RegisterAssetTypeActions(const TSharedRef<IAssetTypeActions>& NewActions)
 {
-	const UClass* SupportedClass = NewActions->GetSupportedClass();
-	NewActions->SetSupported(SupportedClass && AssetClassBlacklist->PassesFilter(SupportedClass->GetFName()));
+	bool bSupported = false;
+	if (const UClass* SupportedClass = NewActions->GetSupportedClass())
+	{
+		bSupported = AssetClassBlacklist->PassesFilter(SupportedClass->GetFName());
+	}
+	else
+	{
+		bSupported = !NewActions->GetFilterName().IsNone();
+	}
+
+	NewActions->SetSupported(bSupported);
 
 	AssetTypeActionsList.Add(NewActions);
 }
@@ -3365,8 +3374,17 @@ void UAssetToolsImpl::AssetClassBlacklistChanged()
 {
 	for (TSharedRef<IAssetTypeActions>& ActionsIt : AssetTypeActionsList)
 	{
-		const UClass* SupportedClass = ActionsIt->GetSupportedClass();
-		ActionsIt->SetSupported(SupportedClass && AssetClassBlacklist->PassesFilter(SupportedClass->GetFName()));
+		bool bSupported = false;
+		if (const UClass* SupportedClass = ActionsIt->GetSupportedClass())
+		{
+			bSupported = AssetClassBlacklist->PassesFilter(SupportedClass->GetFName());
+		}
+		else
+		{
+			bSupported = !ActionsIt->GetFilterName().IsNone();
+		}
+
+		ActionsIt->SetSupported(bSupported);
 	}
 }
 
diff --git a/Engine/Source/Developer/AssetTools/Public/AssetTypeActions_Base.h b/Engine/Source/Developer/AssetTools/Public/AssetTypeActions_Base.h
--- a/Engine/Source/Developer/AssetTools/Public/AssetTypeActions_Base.h
+++ b/Engine/Source/Developer/AssetTools/Public/AssetTypeActions_Base.h
@@ -148,7 +148,7 @@ public:
 	virtual void BuildBackendFilter(FARFilter& InFilter) override
 	{
 		// Add the supported class for this type to a filter
-		InFilter.ClassNames.Add(GetSupportedClass()->GetFName());
+		InFilter.ClassNames.Add(GetFilterName());
 		InFilter.bRecursiveClasses = true;
 	}
 	
@@ -167,6 +167,11 @@ public:
 		return bIsSupported;
 	}
 
+	virtual FName GetFilterName() const override
+	{
+		return GetSupportedClass()->GetFName();
+	}
+
 protected:
 
 	// Here are some convenience functions for common asset type actions logic
diff --git a/Engine/Source/Developer/AssetTools/Public/IAssetTypeActions.h b/Engine/Source/Developer/AssetTools/Public/IAssetTypeActions.h
--- a/Engine/Source/Developer/AssetTools/Public/IAssetTypeActions.h
+++ b/Engine/Source/Developer/AssetTools/Public/IAssetTypeActions.h
@@ -73,6 +73,9 @@ public:
 	/** Returns true if this class can be used as a filter in the content browser */
 	virtual bool CanFilter() = 0;
 
+	/** Returns name to use for filter in the content browser */
+	virtual FName GetFilterName() const = 0;
+
 	/** Returns true if this class can be localized */
 	virtual bool CanLocalize() const = 0;
 
diff --git a/Engine/Source/Editor/ContentBrowser/Private/SFilterList.cpp b/Engine/Source/Editor/ContentBrowser/Private/SFilterList.cpp
--- a/Engine/Source/Editor/ContentBrowser/Private/SFilterList.cpp
+++ b/Engine/Source/Editor/ContentBrowser/Private/SFilterList.cpp
@@ -646,7 +646,7 @@ FARFilter SFilterList::GetCombinedBackendFilter() const
 				if (TypeActions->CanFilter())
 				{
 					const UClass* TypeClass = TypeActions->GetSupportedClass();
-					if (!CombinedFilter.ClassNames.Contains(TypeClass->GetFName()))
+					if (TypeClass && !CombinedFilter.ClassNames.Contains(TypeClass->GetFName()))
 					{
 						CombinedFilter.RecursiveClassesExclusionSet.Add(TypeClass->GetFName());
 					}
@@ -856,7 +856,7 @@ void SFilterList::SaveSettings(const FString& IniFilename, const FString& IniSec
 				ActiveTypeFilterString += TEXT(",");
 			}
 
-			const FString FilterName = Filter->GetAssetTypeActions().Pin()->GetSupportedClass()->GetName();
+			const FString FilterName = Filter->GetAssetTypeActions().Pin()->GetFilterName().ToString();
 			ActiveTypeFilterString += FilterName;
 
 			if ( Filter->IsEnabled() )
@@ -927,12 +927,12 @@ void SFilterList::LoadSettings(const FString& IniFilename, const FString& IniSec
 			const TWeakPtr<IAssetTypeActions>& TypeActions = *TypeActionsIt;
 			if ( TypeActions.IsValid() && TypeActions.Pin()->CanFilter() && !IsAssetTypeActionsInUse(TypeActions) )
 			{
-				const FString& ClassName = TypeActions.Pin()->GetSupportedClass()->GetName();
-				if ( TypeFilterNames.Contains(ClassName) )
+				const FString FilterName = TypeActions.Pin()->GetFilterName().ToString();
+				if ( TypeFilterNames.Contains(FilterName) )
 				{
 					TSharedRef<SFilter> NewFilter = AddFilter(TypeActions);
 
-					if ( EnabledTypeFilterNames.Contains(ClassName) )
+					if ( EnabledTypeFilterNames.Contains(FilterName) )
 					{
 						NewFilter->SetEnabled(true, false);
 					}
