There was an issue converting from a relative path to a full path if the directories were named similarly

diff --git a/Engine/Plugins/Editor/ContentBrowser/ContentBrowserFileDataSource/Source/ContentBrowserFileDataSource/Private/ContentBrowserFileDataSource.cpp b/Engine/Plugins/Editor/ContentBrowser/ContentBrowserFileDataSource/Source/ContentBrowserFileDataSource/Private/ContentBrowserFileDataSource.cpp
--- a/Engine/Plugins/Editor/ContentBrowser/ContentBrowserFileDataSource/Source/ContentBrowserFileDataSource/Private/ContentBrowserFileDataSource.cpp
+++ b/Engine/Plugins/Editor/ContentBrowser/ContentBrowserFileDataSource/Source/ContentBrowserFileDataSource/Private/ContentBrowserFileDataSource.cpp
@@ -1545,8 +1545,8 @@ bool UContentBrowserFileDataSource::IsKnownFileMount(const FName InMountPath, FS
 		{
 			if (OutDiskPath)
 			{
-				*OutDiskPath = MountPathStrView;
-				OutDiskPath->ReplaceInline(*FileMountMountRootStr, *RegisteredFileMount.Value.DiskPath);
+					*OutDiskPath = MountPathStrView;
+					*OutDiskPath = FPaths::ConvertRelativePathToFull(*FileMountMountRootStr, *RegisteredFileMount.Value.DiskPath);
 			}
 
 			return true;
