Script cannot be added to a project

There was an issue converting from a relative path to a full path if the directories were named similarly, the fix for this caused the later part of the directory to get missed

#rnx

diff --git a/Engine/Plugins/Editor/ContentBrowser/ContentBrowserFileDataSource/Source/ContentBrowserFileDataSource/Private/ContentBrowserFileDataSource.cpp b/Engine/Plugins/Editor/ContentBrowser/ContentBrowserFileDataSource/Source/ContentBrowserFileDataSource/Private/ContentBrowserFileDataSource.cpp
--- a/Engine/Plugins/Editor/ContentBrowser/ContentBrowserFileDataSource/Source/ContentBrowserFileDataSource/Private/ContentBrowserFileDataSource.cpp
+++ b/Engine/Plugins/Editor/ContentBrowser/ContentBrowserFileDataSource/Source/ContentBrowserFileDataSource/Private/ContentBrowserFileDataSource.cpp
@@ -1545,10 +1545,12 @@ bool UContentBrowserFileDataSource::IsKnownFileMount(const FName InMountPath, FS
 		{
 			if (OutDiskPath)
 			{
-					*OutDiskPath = MountPathStrView;
-					*OutDiskPath = FPaths::ConvertRelativePathToFull(*FileMountMountRootStr, *RegisteredFileMount.Value.DiskPath);
-			}
+				*OutDiskPath = MountPathStrView;
+				OutDiskPath->RightChopInline(FileMountMountRootStr.Len());
 
+				FString Path = FPaths::ConvertRelativePathToFull(*FileMountMountRootStr, *RegisteredFileMount.Value.DiskPath);
+				*OutDiskPath = Path.Append(*OutDiskPath); 
+			}
 			return true;
 		}
 	}
