Removing ToolMenu dynamic section added by a ContentBrowserFileDataSource when that data source is shutdown

diff --git a/Engine/Plugins/Editor/ContentBrowser/ContentBrowserFileDataSource/Source/ContentBrowserFileDataSource/Private/ContentBrowserFileDataSource.cpp b/Engine/Plugins/Editor/ContentBrowser/ContentBrowserFileDataSource/Source/ContentBrowserFileDataSource/Private/ContentBrowserFileDataSource.cpp
--- a/Engine/Plugins/Editor/ContentBrowser/ContentBrowserFileDataSource/Source/ContentBrowserFileDataSource/Private/ContentBrowserFileDataSource.cpp
+++ b/Engine/Plugins/Editor/ContentBrowser/ContentBrowserFileDataSource/Source/ContentBrowserFileDataSource/Private/ContentBrowserFileDataSource.cpp
@@ -391,6 +391,16 @@ void UContentBrowserFileDataSource::Shutdown()
 		}
 	}
 
+	{
+		if (UToolMenus* ToolMenus = UToolMenus::Get())
+		{
+			if (UToolMenu* Menu = ToolMenus->ExtendMenu("ContentBrowser.AddNewContextMenu"))
+			{
+				Menu->RemoveSection(*FString::Printf(TEXT("DynamicSection_DataSource_%s"), *GetName()));
+			}
+		}
+	}
+
 	if (BackgroundDiscovery)
 	{
 		BackgroundDiscovery->EnsureCompletion();
