IAssetTypeActions: Add a function to get the FTopLevelAssetPath

diff --git a/Engine/Plugins/Editor/ContentBrowser/ContentBrowserFileDataSource/Source/ContentBrowserFileDataSource/Private/ContentBrowserFileDataSource.cpp b/Engine/Plugins/Editor/ContentBrowser/ContentBrowserFileDataSource/Source/ContentBrowserFileDataSource/Private/ContentBrowserFileDataSource.cpp
--- a/Engine/Plugins/Editor/ContentBrowser/ContentBrowserFileDataSource/Source/ContentBrowserFileDataSource/Private/ContentBrowserFileDataSource.cpp
+++ b/Engine/Plugins/Editor/ContentBrowser/ContentBrowserFileDataSource/Source/ContentBrowserFileDataSource/Private/ContentBrowserFileDataSource.cpp
@@ -56,12 +56,14 @@ public:
 	virtual UClass* GetSupportedClass() const override;
 	virtual FColor GetTypeColor() const override;
 	virtual FName GetFilterName() const override;
+	virtual FTopLevelAssetPath GetClassPathName() const override;
 	// End IAssetTypeActions interface
 
 	FName FilterName;
 	FText Name;
 	FText Description;
 	FColor TypeColor;
+	FTopLevelAssetPath ClassPathName;
 };
 
 FText FAssetTypeActions_FileDataSource::GetAssetDescription(const struct FAssetData& AssetData) const
@@ -89,6 +91,11 @@ FName FAssetTypeActions_FileDataSource::GetFilterName() const
 	return FilterName;
 }
 
+FTopLevelAssetPath FAssetTypeActions_FileDataSource::GetClassPathName() const
+{
+	return ClassPathName;
+}
+
 class FContentBrowserFileDataDiscovery : public FRunnable
 {
 public:
@@ -374,6 +381,7 @@ void UContentBrowserFileDataSource::Initialize(const ContentBrowserFileData::FFi
 		AssetTypeAction->Description = InFileActions->TypeFullDescription;
 		AssetTypeAction->TypeColor = InFileActions->TypeColor.ToFColor(true);
 		AssetTypeAction->FilterName = *(ContentBrowserFileDataSource::GetFilterNamePrefix() + InFileActions->TypeName.GetAssetName().ToString());
+		AssetTypeAction->ClassPathName = InFileActions->TypeName;
 		AssetTools.RegisterAssetTypeActions(AssetTypeAction);
 		RegisteredAssetTypeActions.Add(AssetTypeAction);
 		return true;
diff --git a/Engine/Source/Developer/AssetTools/Public/AssetTypeActions_Base.h b/Engine/Source/Developer/AssetTools/Public/AssetTypeActions_Base.h
--- a/Engine/Source/Developer/AssetTools/Public/AssetTypeActions_Base.h
+++ b/Engine/Source/Developer/AssetTools/Public/AssetTypeActions_Base.h
@@ -158,9 +158,7 @@ public:
 	virtual void BuildBackendFilter(FARFilter& InFilter) override
 	{
 		// Add the supported class for this type to a filter
-		FString FilterString(GetFilterName().ToString());
-		checkf(!FPackageName::IsShortPackageName(FilterString), TEXT("Short class names are not supported: %s"), *FilterString);
-		InFilter.ClassPaths.Add(FTopLevelAssetPath(FilterString));
+		InFilter.ClassPaths.Add(GetClassPathName());
 		InFilter.bRecursiveClasses = true;
 	}
 	
@@ -184,6 +182,11 @@ public:
 		return FName(*GetSupportedClass()->GetPathName());
 	}
 
+	virtual FTopLevelAssetPath GetClassPathName() const override
+	{
+		return GetSupportedClass()->GetClassPathName();
+	}
+
 protected:
 
 	// Here are some convenience functions for common asset type actions logic
diff --git a/Engine/Source/Developer/AssetTools/Public/IAssetTypeActions.h b/Engine/Source/Developer/AssetTools/Public/IAssetTypeActions.h
--- a/Engine/Source/Developer/AssetTools/Public/IAssetTypeActions.h
+++ b/Engine/Source/Developer/AssetTools/Public/IAssetTypeActions.h
@@ -150,4 +150,9 @@ public:
 
 	/** Is this asset type supported in the current session? */
 	virtual bool IsSupported() const = 0;
+
+	/**
+	 * Returns class path name as a package + class FName pair
+	 */
+	virtual FTopLevelAssetPath GetClassPathName() const = 0;
 };
