Fix crash on shutdown in ContentBrowserFileDataSource module

diff --git a/Engine/Plugins/Editor/ContentBrowser/ContentBrowserFileDataSource/Source/ContentBrowserFileDataSource/Private/ContentBrowserFileDataSource.cpp b/Engine/Plugins/Editor/ContentBrowser/ContentBrowserFileDataSource/Source/ContentBrowserFileDataSource/Private/ContentBrowserFileDataSource.cpp
--- a/Engine/Plugins/Editor/ContentBrowser/ContentBrowserFileDataSource/Source/ContentBrowserFileDataSource/Private/ContentBrowserFileDataSource.cpp
+++ b/Engine/Plugins/Editor/ContentBrowser/ContentBrowserFileDataSource/Source/ContentBrowserFileDataSource/Private/ContentBrowserFileDataSource.cpp
@@ -360,6 +360,7 @@ void UContentBrowserFileDataSource::Initialize(const ContentBrowserFileData::FFi
 
 	// Bind the asset specific menu extensions
 	{
+		FToolMenuOwnerScoped MenuOwnerScoped(GetFName());
 		if (UToolMenu* Menu = UToolMenus::Get()->ExtendMenu("ContentBrowser.AddNewContextMenu"))
 		{
 			Menu->AddDynamicSection(*FString::Printf(TEXT("DynamicSection_DataSource_%s"), *GetName()), FNewToolMenuDelegate::CreateLambda([WeakThis = TWeakObjectPtr<UContentBrowserFileDataSource>(this)](UToolMenu* InMenu)
@@ -399,15 +400,7 @@ void UContentBrowserFileDataSource::Shutdown()
 		}
 	}
 
-	{
-		if (UToolMenus* ToolMenus = UToolMenus::Get())
-		{
-			if (UToolMenu* Menu = ToolMenus->ExtendMenu("ContentBrowser.AddNewContextMenu"))
-			{
-				Menu->RemoveSection(*FString::Printf(TEXT("DynamicSection_DataSource_%s"), *GetName()));
-			}
-		}
-	}
+	UToolMenus::UnregisterOwner(GetFName());
 
 	if (BackgroundDiscovery)
 	{
