Fix crash in ContentBrowser when the -nothreading option is on.

diff --git a/Engine/Plugins/Editor/ContentBrowser/ContentBrowserFileDataSource/Source/ContentBrowserFileDataSource/Private/ContentBrowserFileDataSource.cpp b/Engine/Plugins/Editor/ContentBrowser/ContentBrowserFileDataSource/Source/ContentBrowserFileDataSource/Private/ContentBrowserFileDataSource.cpp
--- a/Engine/Plugins/Editor/ContentBrowser/ContentBrowserFileDataSource/Source/ContentBrowserFileDataSource/Private/ContentBrowserFileDataSource.cpp
+++ b/Engine/Plugins/Editor/ContentBrowser/ContentBrowserFileDataSource/Source/ContentBrowserFileDataSource/Private/ContentBrowserFileDataSource.cpp
@@ -21,6 +21,8 @@
 
 #define LOCTEXT_NAMESPACE "ContentBrowserFileDataSource"
 
+DEFINE_LOG_CATEGORY_STATIC(LogContentBrowserFileDataSource, Warning, Warning);
+
 namespace ContentBrowserFileDataSource
 {
 	FStringView GetRootOfPath(FStringView InPath)
@@ -356,7 +358,15 @@ void UContentBrowserFileDataSource::Initialize(const ContentBrowserFileData::FFi
 	Super::Initialize(InAutoRegister);
 
 	Config = InConfig;
-	BackgroundDiscovery = MakeShared<FContentBrowserFileDataDiscovery>(&Config);
+
+	if (FPlatformProcess::SupportsMultithreading())
+	{
+		BackgroundDiscovery = MakeShared<FContentBrowserFileDataDiscovery>(&Config);
+	}
+	else
+	{
+		UE_LOG(LogContentBrowserFileDataSource, Error, TEXT("UContentBrowserFileDataSource '%s': Unable to start data source, support for threads is required."), *GetFullName());
+	}
 
 	// Bind the asset specific menu extensions
 	{
