Adding some cpu traces for the content browser data source filter compilation.

#rnx

diff --git a/Engine/Plugins/Editor/ContentBrowser/ContentBrowserAliasDataSoure/Source/ContentBrowserAliasDataSource/Private/ContentBrowserAliasDataSource.cpp b/Engine/Plugins/Editor/ContentBrowser/ContentBrowserAliasDataSoure/Source/ContentBrowserAliasDataSource/Private/ContentBrowserAliasDataSource.cpp
--- a/Engine/Plugins/Editor/ContentBrowser/ContentBrowserAliasDataSoure/Source/ContentBrowserAliasDataSource/Private/ContentBrowserAliasDataSource.cpp
+++ b/Engine/Plugins/Editor/ContentBrowser/ContentBrowserAliasDataSoure/Source/ContentBrowserAliasDataSource/Private/ContentBrowserAliasDataSource.cpp
@@ -64,6 +64,8 @@ void UContentBrowserAliasDataSource::BuildRootPathVirtualTree()
 
 void UContentBrowserAliasDataSource::CompileFilter(const FName InPath, const FContentBrowserDataFilter& InFilter, FContentBrowserDataCompiledFilter& OutCompiledFilter)
 {
+	TRACE_CPUPROFILER_EVENT_SCOPE(UContentBrowserAliasDataSource::CompileFilter);
+
 	UContentBrowserAssetDataSource::FAssetFilterInputParams Params;
 	if (UContentBrowserAssetDataSource::PopulateAssetFilterInputParams(Params, this, AssetRegistry, InFilter, OutCompiledFilter, &FCollectionManagerModule::GetModule().Get()))
 	{
diff --git a/Engine/Plugins/Editor/ContentBrowser/ContentBrowserAssetDataSource/Source/ContentBrowserAssetDataSource/Private/ContentBrowserAssetDataSource.cpp b/Engine/Plugins/Editor/ContentBrowser/ContentBrowserAssetDataSource/Source/ContentBrowserAssetDataSource/Private/ContentBrowserAssetDataSource.cpp
--- a/Engine/Plugins/Editor/ContentBrowser/ContentBrowserAssetDataSource/Source/ContentBrowserAssetDataSource/Private/ContentBrowserAssetDataSource.cpp
+++ b/Engine/Plugins/Editor/ContentBrowser/ContentBrowserAssetDataSource/Source/ContentBrowserAssetDataSource/Private/ContentBrowserAssetDataSource.cpp
@@ -932,6 +932,8 @@ PRAGMA_ENABLE_DEPRECATION_WARNINGS
 
 void UContentBrowserAssetDataSource::CompileFilter(const FName InPath, const FContentBrowserDataFilter& InFilter, FContentBrowserDataCompiledFilter& OutCompiledFilter)
 {
+	TRACE_CPUPROFILER_EVENT_SCOPE(UContentBrowserAssetDataSource::CompileFilter);
+
 	FAssetFilterInputParams Params;
 	if (PopulateAssetFilterInputParams(Params, this, AssetRegistry, InFilter, OutCompiledFilter, CollectionManager))
 	{
diff --git a/Engine/Plugins/Editor/ContentBrowser/ContentBrowserClassDataSource/Source/ContentBrowserClassDataSource/Private/ContentBrowserClassDataSource.cpp b/Engine/Plugins/Editor/ContentBrowser/ContentBrowserClassDataSource/Source/ContentBrowserClassDataSource/Private/ContentBrowserClassDataSource.cpp
--- a/Engine/Plugins/Editor/ContentBrowser/ContentBrowserClassDataSource/Source/ContentBrowserClassDataSource/Private/ContentBrowserClassDataSource.cpp
+++ b/Engine/Plugins/Editor/ContentBrowser/ContentBrowserClassDataSource/Source/ContentBrowserClassDataSource/Private/ContentBrowserClassDataSource.cpp
@@ -97,6 +97,8 @@ bool UContentBrowserClassDataSource::RootClassPathPassesFilter(const FName InRoo
 
 void UContentBrowserClassDataSource::CompileFilter(const FName InPath, const FContentBrowserDataFilter& InFilter, FContentBrowserDataCompiledFilter& OutCompiledFilter)
 {
+	TRACE_CPUPROFILER_EVENT_SCOPE(UContentBrowserClassDataSource::CompileFilter);
+
 	const FContentBrowserDataClassFilter* ClassFilter = InFilter.ExtraFilters.FindFilter<FContentBrowserDataClassFilter>();
 	const FContentBrowserDataCollectionFilter* CollectionFilter = InFilter.ExtraFilters.FindFilter<FContentBrowserDataCollectionFilter>();
 
diff --git a/Engine/Plugins/Editor/ContentBrowser/ContentBrowserFileDataSource/Source/ContentBrowserFileDataSource/Private/ContentBrowserFileDataSource.cpp b/Engine/Plugins/Editor/ContentBrowser/ContentBrowserFileDataSource/Source/ContentBrowserFileDataSource/Private/ContentBrowserFileDataSource.cpp
--- a/Engine/Plugins/Editor/ContentBrowser/ContentBrowserFileDataSource/Source/ContentBrowserFileDataSource/Private/ContentBrowserFileDataSource.cpp
+++ b/Engine/Plugins/Editor/ContentBrowser/ContentBrowserFileDataSource/Source/ContentBrowserFileDataSource/Private/ContentBrowserFileDataSource.cpp
@@ -565,6 +565,8 @@ void UContentBrowserFileDataSource::BuildRootPathVirtualTree()
 
 void UContentBrowserFileDataSource::CompileFilter(const FName InPath, const FContentBrowserDataFilter& InFilter, FContentBrowserDataCompiledFilter& OutCompiledFilter)
 {
+	TRACE_CPUPROFILER_EVENT_SCOPE(UContentBrowserFileDataSource::CompileFilter);
+
 	const FContentBrowserDataObjectFilter* ObjectFilter = InFilter.ExtraFilters.FindFilter<FContentBrowserDataObjectFilter>();
 	const FContentBrowserDataPackageFilter* PackageFilter = InFilter.ExtraFilters.FindFilter<FContentBrowserDataPackageFilter>();
 	const FContentBrowserDataClassFilter* ClassFilter = InFilter.ExtraFilters.FindFilter<FContentBrowserDataClassFilter>();
diff --git a/Engine/Source/Editor/ContentBrowser/Private/SAssetView.cpp b/Engine/Source/Editor/ContentBrowser/Private/SAssetView.cpp
--- a/Engine/Source/Editor/ContentBrowser/Private/SAssetView.cpp
+++ b/Engine/Source/Editor/ContentBrowser/Private/SAssetView.cpp
@@ -4709,6 +4709,8 @@ void SAssetView::OverrideShowLocalizedContent()
 
 void SAssetView::HandleItemDataUpdated(TArrayView<const FContentBrowserItemDataUpdate> InUpdatedItems)
 {
+	TRACE_CPUPROFILER_EVENT_SCOPE(SAssetView::HandleItemDataUpdated);
+
 	if (InUpdatedItems.Num() == 0)
 	{
 		return;
