Updated headers using UnrealCodeFixup to make sure dllstorage is on methods/staticvars instead of types. Part 2/n

diff --git a/Engine/Plugins/Editor/ContentBrowser/ContentBrowserAssetDataSource/Source/ContentBrowserAssetDataSource/Public/ContentBrowserAssetDataPayload.h b/Engine/Plugins/Editor/ContentBrowser/ContentBrowserAssetDataSource/Source/ContentBrowserAssetDataSource/Public/ContentBrowserAssetDataPayload.h
--- a/Engine/Plugins/Editor/ContentBrowser/ContentBrowserAssetDataSource/Source/ContentBrowserAssetDataSource/Public/ContentBrowserAssetDataPayload.h
+++ b/Engine/Plugins/Editor/ContentBrowser/ContentBrowserAssetDataSource/Source/ContentBrowserAssetDataSource/Public/ContentBrowserAssetDataPayload.h
@@ -6,12 +6,14 @@
 #include "AssetRegistry/AssetData.h"
 #include "UObject/GCObject.h"
 
+#define UE_API CONTENTBROWSERASSETDATASOURCE_API
+
 class UAssetDefinition;
 class IAssetTypeActions;
 class FAssetThumbnail;
 class UFactory;
 
-class CONTENTBROWSERASSETDATASOURCE_API FContentBrowserAssetFolderItemDataPayload : public IContentBrowserItemDataPayload
+class FContentBrowserAssetFolderItemDataPayload : public IContentBrowserItemDataPayload
 {
 public:
 	explicit FContentBrowserAssetFolderItemDataPayload(const FName InInternalPath)
@@ -24,7 +26,7 @@ public:
 		return InternalPath;
 	}
 
-	const FString& GetFilename() const;
+	UE_API const FString& GetFilename() const;
 
 private:
 	FName InternalPath;
@@ -33,34 +35,34 @@ private:
 	mutable FString CachedFilename;
 };
 
-class CONTENTBROWSERASSETDATASOURCE_API FContentBrowserAssetFileItemDataPayload : public IContentBrowserItemDataPayload
+class FContentBrowserAssetFileItemDataPayload : public IContentBrowserItemDataPayload
 {
 public:
-	explicit FContentBrowserAssetFileItemDataPayload(FAssetData&& InAssetData);
-	explicit FContentBrowserAssetFileItemDataPayload(const FAssetData& InAssetData);
+	UE_API explicit FContentBrowserAssetFileItemDataPayload(FAssetData&& InAssetData);
+	UE_API explicit FContentBrowserAssetFileItemDataPayload(const FAssetData& InAssetData);
 
 	const FAssetData& GetAssetData() const
 	{
 		return AssetData;
 	}
 
-	UPackage* GetPackage(const bool bTryRecacheIfNull = false) const;
+	UE_API UPackage* GetPackage(const bool bTryRecacheIfNull = false) const;
 
 	// LoadTags(optional) allows passing specific tags to the linker when loading the asset package (@see ULevel::LoadAllExternalObjectsTag for an example usage)
-	UPackage* LoadPackage(TSet<FName> LoadTags = {}) const;
+	UE_API UPackage* LoadPackage(TSet<FName> LoadTags = {}) const;
 
-	UObject* GetAsset(const bool bTryRecacheIfNull = false) const;
+	UE_API UObject* GetAsset(const bool bTryRecacheIfNull = false) const;
 
 	//  LoadTags (optional) allows passing specific tags to the linker when loading the asset (@see ULevel::LoadAllExternalObjectsTag for an example usage)
-	UObject* LoadAsset(TSet<FName> LoadTags = {}) const;
+	UE_API UObject* LoadAsset(TSet<FName> LoadTags = {}) const;
 
-	TSharedPtr<IAssetTypeActions> GetAssetTypeActions() const;
+	UE_API TSharedPtr<IAssetTypeActions> GetAssetTypeActions() const;
 	
-	const UAssetDefinition* GetAssetDefinition() const;
+	UE_API const UAssetDefinition* GetAssetDefinition() const;
 
-	const FString& GetFilename() const;
+	UE_API const FString& GetFilename() const;
 
-	void UpdateThumbnail(FAssetThumbnail& InThumbnail) const;
+	UE_API void UpdateThumbnail(FAssetThumbnail& InThumbnail) const;
 
 private:
 	FAssetData AssetData;
@@ -80,13 +82,13 @@ private:
 	mutable bool bHasCachedFilename = false;
 	mutable FString CachedFilename;
 
-	void FlushCaches() const;
+	UE_API void FlushCaches() const;
 };
 
-class CONTENTBROWSERASSETDATASOURCE_API FContentBrowserAssetFileItemDataPayload_Creation : public FContentBrowserAssetFileItemDataPayload, public FGCObject
+class FContentBrowserAssetFileItemDataPayload_Creation : public FContentBrowserAssetFileItemDataPayload, public FGCObject
 {
 public:
-	FContentBrowserAssetFileItemDataPayload_Creation(FAssetData&& InAssetData, UClass* InAssetClass, UFactory* InFactory);
+	UE_API FContentBrowserAssetFileItemDataPayload_Creation(FAssetData&& InAssetData, UClass* InAssetClass, UFactory* InFactory);
 
 	//~ FGCObject interface
 	virtual void AddReferencedObjects(FReferenceCollector& Collector) override
@@ -117,10 +119,10 @@ private:
 	TObjectPtr<UFactory> Factory = nullptr;
 };
 
-class CONTENTBROWSERASSETDATASOURCE_API FContentBrowserAssetFileItemDataPayload_Duplication : public FContentBrowserAssetFileItemDataPayload
+class FContentBrowserAssetFileItemDataPayload_Duplication : public FContentBrowserAssetFileItemDataPayload
 {
 public:
-	FContentBrowserAssetFileItemDataPayload_Duplication(FAssetData&& InAssetData, TWeakObjectPtr<UObject> InSourceObject);
+	UE_API FContentBrowserAssetFileItemDataPayload_Duplication(FAssetData&& InAssetData, TWeakObjectPtr<UObject> InSourceObject);
 
 	UObject* GetSourceObject() const
 	{
@@ -132,21 +134,21 @@ private:
 	TWeakObjectPtr<UObject> SourceObject;
 };
 
-class CONTENTBROWSERASSETDATASOURCE_API FContentBrowserUnsupportedAssetFileItemDataPayload : public IContentBrowserItemDataPayload
+class FContentBrowserUnsupportedAssetFileItemDataPayload : public IContentBrowserItemDataPayload
 {
 public:
 	// Unsupported asset file but it does have an asset data
-	explicit FContentBrowserUnsupportedAssetFileItemDataPayload(FAssetData&& InAssetData);
-	explicit FContentBrowserUnsupportedAssetFileItemDataPayload(const FAssetData& InAssetData);
+	UE_API explicit FContentBrowserUnsupportedAssetFileItemDataPayload(FAssetData&& InAssetData);
+	UE_API explicit FContentBrowserUnsupportedAssetFileItemDataPayload(const FAssetData& InAssetData);
 
-	const FAssetData* GetAssetDataIfAvailable() const;
+	UE_API const FAssetData* GetAssetDataIfAvailable() const;
 
-	const FString& GetFilename() const;
+	UE_API const FString& GetFilename() const;
 
-	UPackage* GetPackage() const;
+	UE_API UPackage* GetPackage() const;
 
 private:
-	void FlushCaches() const;
+	UE_API void FlushCaches() const;
 
 
 	TUniquePtr<FAssetData> OptionalAssetData;
@@ -157,3 +159,5 @@ private:
 	mutable bool bHasCachedFilename = false;
 	mutable FString CachedFilename;
 };
+
+#undef UE_API
diff --git a/Engine/Plugins/Editor/ContentBrowser/ContentBrowserAssetDataSource/Source/ContentBrowserAssetDataSource/Public/NewAssetContextMenu.h b/Engine/Plugins/Editor/ContentBrowser/ContentBrowserAssetDataSource/Source/ContentBrowserAssetDataSource/Public/NewAssetContextMenu.h
--- a/Engine/Plugins/Editor/ContentBrowser/ContentBrowserAssetDataSource/Source/ContentBrowserAssetDataSource/Public/NewAssetContextMenu.h
+++ b/Engine/Plugins/Editor/ContentBrowser/ContentBrowserAssetDataSource/Source/ContentBrowserAssetDataSource/Public/NewAssetContextMenu.h
@@ -5,17 +5,19 @@
 #include "AssetTypeCategories.h"
 #include "Framework/Commands/UIAction.h"
 
+#define UE_API CONTENTBROWSERASSETDATASOURCE_API
+
 class UToolMenu;
 struct FCategorySubMenuItem;
 
-class CONTENTBROWSERASSETDATASOURCE_API FNewAssetContextMenu
+class FNewAssetContextMenu
 {
 public:
 	DECLARE_DELEGATE_OneParam(FOnImportAssetRequested, const FName /*SelectedPath*/);
 	DECLARE_DELEGATE_TwoParams(FOnNewAssetRequested, const FName /*SelectedPath*/, TWeakObjectPtr<UClass> /*FactoryClass*/);
 
 	/** Makes the context menu widget */
-	static void MakeContextMenu(
+	static UE_API void MakeContextMenu(
 		UToolMenu* Menu,
 		const TArray<FName>& InSelectedAssetPaths,
 		const FOnImportAssetRequested& InOnImportAssetRequested,
@@ -24,12 +26,14 @@ public:
 
 private:
 	/** Handle creating a new asset from an asset category */
-	static void CreateNewAssetMenuCategory(UToolMenu* Menu, FName SectionName, EAssetTypeCategories::Type AssetTypeCategory, FName InPath, FOnNewAssetRequested InOnNewAssetRequested, FCanExecuteAction InCanExecuteAction);
-	static void CreateNewAssetMenus(UToolMenu* Menu, FName SectionName, TSharedPtr<FCategorySubMenuItem> SubMenuData, FName InPath, FOnNewAssetRequested InOnNewAssetRequested, FCanExecuteAction InCanExecuteAction);
+	static UE_API void CreateNewAssetMenuCategory(UToolMenu* Menu, FName SectionName, EAssetTypeCategories::Type AssetTypeCategory, FName InPath, FOnNewAssetRequested InOnNewAssetRequested, FCanExecuteAction InCanExecuteAction);
+	static UE_API void CreateNewAssetMenus(UToolMenu* Menu, FName SectionName, TSharedPtr<FCategorySubMenuItem> SubMenuData, FName InPath, FOnNewAssetRequested InOnNewAssetRequested, FCanExecuteAction InCanExecuteAction);
 
 	/** Handle when the "Import" button is clicked */
-	static void ExecuteImportAsset(FOnImportAssetRequested InOnImportAssetRequested, FName InPath);
+	static UE_API void ExecuteImportAsset(FOnImportAssetRequested InOnImportAssetRequested, FName InPath);
 
 	/** Create a new asset using the specified factory at the specified path */
-	static void ExecuteNewAsset(FOnNewAssetRequested InOnNewAssetRequested, FName InPath, TWeakObjectPtr<UClass> FactoryClass);
+	static UE_API void ExecuteNewAsset(FOnNewAssetRequested InOnNewAssetRequested, FName InPath, TWeakObjectPtr<UClass> FactoryClass);
 };
+
+#undef UE_API
diff --git a/Engine/Plugins/Editor/ContentBrowser/ContentBrowserClassDataSource/Source/ContentBrowserClassDataSource/Public/ContentBrowserClassDataPayload.h b/Engine/Plugins/Editor/ContentBrowser/ContentBrowserClassDataSource/Source/ContentBrowserClassDataSource/Public/ContentBrowserClassDataPayload.h
--- a/Engine/Plugins/Editor/ContentBrowser/ContentBrowserClassDataSource/Source/ContentBrowserClassDataSource/Public/ContentBrowserClassDataPayload.h
+++ b/Engine/Plugins/Editor/ContentBrowser/ContentBrowserClassDataSource/Source/ContentBrowserClassDataSource/Public/ContentBrowserClassDataPayload.h
@@ -5,9 +5,11 @@
 #include "ContentBrowserItemData.h"
 #include "AssetRegistry/AssetData.h"
 
+#define UE_API CONTENTBROWSERCLASSDATASOURCE_API
+
 class FAssetThumbnail;
 
-class CONTENTBROWSERCLASSDATASOURCE_API FContentBrowserClassFolderItemDataPayload : public IContentBrowserItemDataPayload
+class FContentBrowserClassFolderItemDataPayload : public IContentBrowserItemDataPayload
 {
 public:
 	explicit FContentBrowserClassFolderItemDataPayload(const FName InInternalPath)
@@ -20,7 +22,7 @@ public:
 		return InternalPath;
 	}
 
-	const FString& GetFilename() const;
+	UE_API const FString& GetFilename() const;
 
 private:
 	FName InternalPath;
@@ -29,10 +31,10 @@ private:
 	mutable FString CachedFilename;
 };
 
-class CONTENTBROWSERCLASSDATASOURCE_API FContentBrowserClassFileItemDataPayload : public IContentBrowserItemDataPayload
+class FContentBrowserClassFileItemDataPayload : public IContentBrowserItemDataPayload
 {
 public:
-	FContentBrowserClassFileItemDataPayload(const FName InInternalPath, UClass* InClass);
+	UE_API FContentBrowserClassFileItemDataPayload(const FName InInternalPath, UClass* InClass);
 
 	FName GetInternalPath() const
 	{
@@ -49,9 +51,9 @@ public:
 		return AssetData;
 	}
 
-	const FString& GetFilename() const;
+	UE_API const FString& GetFilename() const;
 
-	void UpdateThumbnail(FAssetThumbnail& InThumbnail) const;
+	UE_API void UpdateThumbnail(FAssetThumbnail& InThumbnail) const;
 
 private:
 	FName InternalPath;
@@ -63,3 +65,5 @@ private:
 	mutable bool bHasCachedFilename = false;
 	mutable FString CachedFilename;
 };
+
+#undef UE_API
diff --git a/Engine/Plugins/Editor/ContentBrowser/ContentBrowserClassDataSource/Source/ContentBrowserClassDataSource/Public/ContentBrowserClassDataSource.h b/Engine/Plugins/Editor/ContentBrowser/ContentBrowserClassDataSource/Source/ContentBrowserClassDataSource/Public/ContentBrowserClassDataSource.h
--- a/Engine/Plugins/Editor/ContentBrowser/ContentBrowserClassDataSource/Source/ContentBrowserClassDataSource/Public/ContentBrowserClassDataSource.h
+++ b/Engine/Plugins/Editor/ContentBrowser/ContentBrowserClassDataSource/Source/ContentBrowserClassDataSource/Public/ContentBrowserClassDataSource.h
@@ -7,6 +7,8 @@
 #include "UObject/Package.h"
 #include "ContentBrowserClassDataSource.generated.h"
 
+#define UE_API CONTENTBROWSERCLASSDATASOURCE_API
+
 class FContentBrowserClassFileItemDataPayload;
 class FContentBrowserClassFolderItemDataPayload;
 
@@ -16,7 +18,7 @@ class FNativeClassHierarchy;
 struct FCollectionNameType;
 
 USTRUCT()
-struct CONTENTBROWSERCLASSDATASOURCE_API FContentBrowserCompiledClassDataFilter
+struct FContentBrowserCompiledClassDataFilter
 {
 	GENERATED_BODY()
 
@@ -28,97 +30,99 @@ public:
 	TSet<FName> ValidFolders;
 };
 
-UCLASS()
-class CONTENTBROWSERCLASSDATASOURCE_API UContentBrowserClassDataSource : public UContentBrowserDataSource
+UCLASS(MinimalAPI)
+class UContentBrowserClassDataSource : public UContentBrowserDataSource
 {
 	GENERATED_BODY()
 
 public:
-	void Initialize(const bool InAutoRegister = true);
+	UE_API void Initialize(const bool InAutoRegister = true);
 
-	virtual void Shutdown() override;
+	UE_API virtual void Shutdown() override;
 
-	virtual void CompileFilter(const FName InPath, const FContentBrowserDataFilter& InFilter, FContentBrowserDataCompiledFilter& OutCompiledFilter) override;
+	UE_API virtual void CompileFilter(const FName InPath, const FContentBrowserDataFilter& InFilter, FContentBrowserDataCompiledFilter& OutCompiledFilter) override;
 
-	virtual void EnumerateItemsMatchingFilter(const FContentBrowserDataCompiledFilter& InFilter, TFunctionRef<bool(FContentBrowserItemData&&)> InCallback) override;
+	UE_API virtual void EnumerateItemsMatchingFilter(const FContentBrowserDataCompiledFilter& InFilter, TFunctionRef<bool(FContentBrowserItemData&&)> InCallback) override;
 
-	virtual void EnumerateItemsAtPath(const FName InPath, const EContentBrowserItemTypeFilter InItemTypeFilter, TFunctionRef<bool(FContentBrowserItemData&&)> InCallback) override;
+	UE_API virtual void EnumerateItemsAtPath(const FName InPath, const EContentBrowserItemTypeFilter InItemTypeFilter, TFunctionRef<bool(FContentBrowserItemData&&)> InCallback) override;
 
-	virtual bool EnumerateItemsForObjects(const TArrayView<UObject*> InObjects, TFunctionRef<bool(FContentBrowserItemData&&)> InCallback) override;
+	UE_API virtual bool EnumerateItemsForObjects(const TArrayView<UObject*> InObjects, TFunctionRef<bool(FContentBrowserItemData&&)> InCallback) override;
 
-	virtual bool IsFolderVisible(const FName InPath, const EContentBrowserIsFolderVisibleFlags InFlags, const FContentBrowserFolderContentsFilter& InContentsFilter) override;
+	UE_API virtual bool IsFolderVisible(const FName InPath, const EContentBrowserIsFolderVisibleFlags InFlags, const FContentBrowserFolderContentsFilter& InContentsFilter) override;
 
-	virtual bool DoesItemPassFilter(const FContentBrowserItemData& InItem, const FContentBrowserDataCompiledFilter& InFilter) override;
+	UE_API virtual bool DoesItemPassFilter(const FContentBrowserItemData& InItem, const FContentBrowserDataCompiledFilter& InFilter) override;
 
-	virtual bool GetItemAttribute(const FContentBrowserItemData& InItem, const bool InIncludeMetaData, const FName InAttributeKey, FContentBrowserItemDataAttributeValue& OutAttributeValue) override;
+	UE_API virtual bool GetItemAttribute(const FContentBrowserItemData& InItem, const bool InIncludeMetaData, const FName InAttributeKey, FContentBrowserItemDataAttributeValue& OutAttributeValue) override;
 
-	virtual bool GetItemAttributes(const FContentBrowserItemData& InItem, const bool InIncludeMetaData, FContentBrowserItemDataAttributeValues& OutAttributeValues) override;
+	UE_API virtual bool GetItemAttributes(const FContentBrowserItemData& InItem, const bool InIncludeMetaData, FContentBrowserItemDataAttributeValues& OutAttributeValues) override;
 
-	virtual bool GetItemPhysicalPath(const FContentBrowserItemData& InItem, FString& OutDiskPath) override;
+	UE_API virtual bool GetItemPhysicalPath(const FContentBrowserItemData& InItem, FString& OutDiskPath) override;
 
-	virtual bool CanEditItem(const FContentBrowserItemData& InItem, FText* OutErrorMsg) override;
+	UE_API virtual bool CanEditItem(const FContentBrowserItemData& InItem, FText* OutErrorMsg) override;
 
-	virtual bool EditItem(const FContentBrowserItemData& InItem) override;
+	UE_API virtual bool EditItem(const FContentBrowserItemData& InItem) override;
 
-	virtual bool BulkEditItems(TArrayView<const FContentBrowserItemData> InItems) override;
+	UE_API virtual bool BulkEditItems(TArrayView<const FContentBrowserItemData> InItems) override;
 
-	virtual bool AppendItemReference(const FContentBrowserItemData& InItem, FString& InOutStr) override;
+	UE_API virtual bool AppendItemReference(const FContentBrowserItemData& InItem, FString& InOutStr) override;
 
-	virtual bool AppendItemObjectPath(const FContentBrowserItemData& InItem, FString& InOutStr) override;
+	UE_API virtual bool AppendItemObjectPath(const FContentBrowserItemData& InItem, FString& InOutStr) override;
 
-	virtual bool AppendItemPackageName(const FContentBrowserItemData& InItem, FString& InOutStr) override;
+	UE_API virtual bool AppendItemPackageName(const FContentBrowserItemData& InItem, FString& InOutStr) override;
 
-	virtual bool UpdateThumbnail(const FContentBrowserItemData& InItem, FAssetThumbnail& InThumbnail) override;
+	UE_API virtual bool UpdateThumbnail(const FContentBrowserItemData& InItem, FAssetThumbnail& InThumbnail) override;
 
-	virtual bool TryGetCollectionId(const FContentBrowserItemData& InItem, FSoftObjectPath& OutCollectionId) override;
+	UE_API virtual bool TryGetCollectionId(const FContentBrowserItemData& InItem, FSoftObjectPath& OutCollectionId) override;
 
-	virtual bool Legacy_TryGetPackagePath(const FContentBrowserItemData& InItem, FName& OutPackagePath) override;
+	UE_API virtual bool Legacy_TryGetPackagePath(const FContentBrowserItemData& InItem, FName& OutPackagePath) override;
 
-	virtual bool Legacy_TryGetAssetData(const FContentBrowserItemData& InItem, FAssetData& OutAssetData) override;
+	UE_API virtual bool Legacy_TryGetAssetData(const FContentBrowserItemData& InItem, FAssetData& OutAssetData) override;
 
-	virtual bool Legacy_TryConvertPackagePathToVirtualPath(const FName InPackagePath, FName& OutPath) override;
+	UE_API virtual bool Legacy_TryConvertPackagePathToVirtualPath(const FName InPackagePath, FName& OutPath) override;
 
-	virtual bool Legacy_TryConvertAssetDataToVirtualPath(const FAssetData& InAssetData, const bool InUseFolderPaths, FName& OutPath) override;
+	UE_API virtual bool Legacy_TryConvertAssetDataToVirtualPath(const FAssetData& InAssetData, const bool InUseFolderPaths, FName& OutPath) override;
 
 protected:
-	TSharedPtr<IAssetTypeActions> GetClassTypeActions();
+	UE_API TSharedPtr<IAssetTypeActions> GetClassTypeActions();
 	
-	virtual void BuildRootPathVirtualTree() override;
-	bool RootClassPathPassesFilter(const FName InRootClassPath, const bool bIncludeEngineClasses, const bool bIncludePluginClasses) const;
+	UE_API virtual void BuildRootPathVirtualTree() override;
+	UE_API bool RootClassPathPassesFilter(const FName InRootClassPath, const bool bIncludeEngineClasses, const bool bIncludePluginClasses) const;
 
 private:
-	bool IsKnownClassPath(const FName InPackagePath) const;
+	UE_API bool IsKnownClassPath(const FName InPackagePath) const;
 
-	bool GetClassPathsForCollections(TArrayView<const FCollectionRef> InCollections, const bool bIncludeChildCollections, TArray<FTopLevelAssetPath>& OutClassPaths);
+	UE_API bool GetClassPathsForCollections(TArrayView<const FCollectionRef> InCollections, const bool bIncludeChildCollections, TArray<FTopLevelAssetPath>& OutClassPaths);
 
-	FContentBrowserItemData CreateClassFolderItem(const FName InFolderPath);
+	UE_API FContentBrowserItemData CreateClassFolderItem(const FName InFolderPath);
 
-	FContentBrowserItemData CreateClassFileItem(UClass* InClass, FNativeClassHierarchyGetClassPathCache& InCache);
+	UE_API FContentBrowserItemData CreateClassFileItem(UClass* InClass, FNativeClassHierarchyGetClassPathCache& InCache);
 
-	FContentBrowserItemData CreateClassFolderItem(const FName InFolderPath, const TSharedPtr<const FNativeClassHierarchyNode>& InFolderNode);
+	UE_API FContentBrowserItemData CreateClassFolderItem(const FName InFolderPath, const TSharedPtr<const FNativeClassHierarchyNode>& InFolderNode);
 
-	FContentBrowserItemData CreateClassFileItem(const FName InClassPath, const TSharedPtr<const FNativeClassHierarchyNode>& InClassNode);
+	UE_API FContentBrowserItemData CreateClassFileItem(const FName InClassPath, const TSharedPtr<const FNativeClassHierarchyNode>& InClassNode);
 
-	TSharedPtr<const FContentBrowserClassFolderItemDataPayload> GetClassFolderItemPayload(const FContentBrowserItemData& InItem) const;
+	UE_API TSharedPtr<const FContentBrowserClassFolderItemDataPayload> GetClassFolderItemPayload(const FContentBrowserItemData& InItem) const;
 
-	TSharedPtr<const FContentBrowserClassFileItemDataPayload> GetClassFileItemPayload(const FContentBrowserItemData& InItem) const;
+	UE_API TSharedPtr<const FContentBrowserClassFileItemDataPayload> GetClassFileItemPayload(const FContentBrowserItemData& InItem) const;
 
-	void OnNewClassRequested(const FName InSelectedPath);
+	UE_API void OnNewClassRequested(const FName InSelectedPath);
 
-	void PopulateAddNewContextMenu(UToolMenu* InMenu);
+	UE_API void PopulateAddNewContextMenu(UToolMenu* InMenu);
 
-	void ConditionalCreateNativeClassHierarchy();
+	UE_API void ConditionalCreateNativeClassHierarchy();
 
-	void OnFoldersAdded(const TArrayView<TSharedRef<const FNativeClassHierarchyNode>> InFolders);
+	UE_API void OnFoldersAdded(const TArrayView<TSharedRef<const FNativeClassHierarchyNode>> InFolders);
 
-	void OnFoldersRemoved(const TArrayView<TSharedRef<const FNativeClassHierarchyNode>> InFolders);
+	UE_API void OnFoldersRemoved(const TArrayView<TSharedRef<const FNativeClassHierarchyNode>> InFolders);
 
-	void OnClassesAdded(const TArrayView<TSharedRef<const FNativeClassHierarchyNode>> InClasses);
+	UE_API void OnClassesAdded(const TArrayView<TSharedRef<const FNativeClassHierarchyNode>> InClasses);
 
-	void OnClassesRemoved(const TArrayView<TSharedRef<const FNativeClassHierarchyNode>> InClasses);
+	UE_API void OnClassesRemoved(const TArrayView<TSharedRef<const FNativeClassHierarchyNode>> InClasses);
 
 	TSharedPtr<FNativeClassHierarchy> NativeClassHierarchy;
 	FNativeClassHierarchyGetClassPathCache NativeClassHierarchyGetClassPathCache;
 
 	TSharedPtr<IAssetTypeActions> ClassTypeActions;
 };
+
+#undef UE_API
diff --git a/Engine/Plugins/Editor/ContentBrowser/ContentBrowserFileDataSource/Source/ContentBrowserFileDataSource/Public/ContentBrowserFileDataCore.h b/Engine/Plugins/Editor/ContentBrowser/ContentBrowserFileDataSource/Source/ContentBrowserFileDataSource/Public/ContentBrowserFileDataCore.h
--- a/Engine/Plugins/Editor/ContentBrowser/ContentBrowserFileDataSource/Source/ContentBrowserFileDataSource/Public/ContentBrowserFileDataCore.h
+++ b/Engine/Plugins/Editor/ContentBrowser/ContentBrowserFileDataSource/Source/ContentBrowserFileDataSource/Public/ContentBrowserFileDataCore.h
@@ -4,6 +4,8 @@
 
 #include "ContentBrowserItemData.h"
 
+#define UE_API CONTENTBROWSERFILEDATASOURCE_API
+
 struct FContentBrowserDataFilter;
 class FContentBrowserFileItemDataPayload;
 class FContentBrowserFileItemDataPayload_Duplication;
@@ -133,11 +135,13 @@ namespace ContentBrowserFileData
 
 	CONTENTBROWSERFILEDATASOURCE_API bool GetFileItemAttributes(const FContentBrowserFileItemDataPayload& InFilePayload, const bool InIncludeMetaData, FContentBrowserItemDataAttributeValues& OutAttributeValues);
 
-	class CONTENTBROWSERFILEDATASOURCE_API FDefaultFileActions
+	class FDefaultFileActions
 	{
 	public:
-		static bool ItemPassesFilter(const FName InFilePath, const FString& InFilename, const FContentBrowserDataFilter& InFilter, const bool bIsFile);
-		static bool GetItemAttribute(const FName InFilePath, const FString& InFilename, const bool InIncludeMetaData, const FName InAttributeKey, FContentBrowserItemDataAttributeValue& OutAttributeValue);
+		static UE_API bool ItemPassesFilter(const FName InFilePath, const FString& InFilename, const FContentBrowserDataFilter& InFilter, const bool bIsFile);
+		static UE_API bool GetItemAttribute(const FName InFilePath, const FString& InFilename, const bool InIncludeMetaData, const FName InAttributeKey, FContentBrowserItemDataAttributeValue& OutAttributeValue);
 	};
 
 }
+
+#undef UE_API
diff --git a/Engine/Plugins/Editor/ContentBrowser/ContentBrowserFileDataSource/Source/ContentBrowserFileDataSource/Public/ContentBrowserFileDataPayload.h b/Engine/Plugins/Editor/ContentBrowser/ContentBrowserFileDataSource/Source/ContentBrowserFileDataSource/Public/ContentBrowserFileDataPayload.h
--- a/Engine/Plugins/Editor/ContentBrowser/ContentBrowserFileDataSource/Source/ContentBrowserFileDataSource/Public/ContentBrowserFileDataPayload.h
+++ b/Engine/Plugins/Editor/ContentBrowser/ContentBrowserFileDataSource/Source/ContentBrowserFileDataSource/Public/ContentBrowserFileDataPayload.h
@@ -5,11 +5,13 @@
 #include "ContentBrowserItemData.h"
 #include "UObject/StructOnScope.h"
 
+#define UE_API CONTENTBROWSERFILEDATASOURCE_API
+
 struct FContentBrowserDataFilter;
 
 namespace ContentBrowserFileData
 {
-	struct CONTENTBROWSERFILEDATASOURCE_API FCommonActions
+	struct FCommonActions
 	{
 		DECLARE_DELEGATE_RetVal_ThreeParams(bool, FCanCreate, const FName /*InDestFolderPath*/, const FString& /*InDestFolder*/, FText* /*OutErrorMsg*/);
 		FCanCreate CanCreate;
@@ -48,11 +50,11 @@ namespace ContentBrowserFileData
 		FPassesFilter PassesFilter;
 	};
 
-	struct CONTENTBROWSERFILEDATASOURCE_API FDirectoryActions : public FCommonActions
+	struct FDirectoryActions : public FCommonActions
 	{
 	};
 
-	struct CONTENTBROWSERFILEDATASOURCE_API FFileActions : public FCommonActions
+	struct FFileActions : public FCommonActions
 	{
 		FString TypeExtension;
 
@@ -86,21 +88,21 @@ namespace ContentBrowserFileData
 		FCanDuplicate CanDuplicate;
 	};
 
-	class CONTENTBROWSERFILEDATASOURCE_API FFileConfigData
+	class FFileConfigData
 	{
 	public:
-		void SetDirectoryActions(const FDirectoryActions& InDirectoryActions);
+		UE_API void SetDirectoryActions(const FDirectoryActions& InDirectoryActions);
 
-		TSharedPtr<const FDirectoryActions> GetDirectoryActions();
+		UE_API TSharedPtr<const FDirectoryActions> GetDirectoryActions();
 
-		void RegisterFileActions(const FFileActions& InFileActions);
+		UE_API void RegisterFileActions(const FFileActions& InFileActions);
 
-		TSharedPtr<const FFileActions> FindFileActionsForExtension(FStringView InTypeExtension) const;
-		TSharedPtr<const FFileActions> FindFileActionsForFilename(FStringView InFilename) const;
+		UE_API TSharedPtr<const FFileActions> FindFileActionsForExtension(FStringView InTypeExtension) const;
+		UE_API TSharedPtr<const FFileActions> FindFileActionsForFilename(FStringView InFilename) const;
 
-		void EnumerateFileActions(TFunctionRef<bool(TSharedRef<const FFileActions>)> InCallback) const;
+		UE_API void EnumerateFileActions(TFunctionRef<bool(TSharedRef<const FFileActions>)> InCallback) const;
 
-		FText GetDiscoveryDescription() const;
+		UE_API FText GetDiscoveryDescription() const;
 
 	private:
 		TSharedPtr<const FDirectoryActions> DirectoryActions;
@@ -109,7 +111,7 @@ namespace ContentBrowserFileData
 	};
 }
 
-class CONTENTBROWSERFILEDATASOURCE_API FContentBrowserFolderItemDataPayload : public IContentBrowserItemDataPayload
+class FContentBrowserFolderItemDataPayload : public IContentBrowserItemDataPayload
 {
 public:
 	explicit FContentBrowserFolderItemDataPayload(const FName InInternalPath, const FString& InFilename, TWeakPtr<const ContentBrowserFileData::FDirectoryActions> InDirectoryActions)
@@ -142,7 +144,7 @@ private:
 	TWeakPtr<const ContentBrowserFileData::FDirectoryActions> DirectoryActions;
 };
 
-class CONTENTBROWSERFILEDATASOURCE_API FContentBrowserFileItemDataPayload : public IContentBrowserItemDataPayload
+class FContentBrowserFileItemDataPayload : public IContentBrowserItemDataPayload
 {
 public:
 	FContentBrowserFileItemDataPayload(const FName InInternalPath, const FString& InFilename, TWeakPtr<const ContentBrowserFileData::FFileActions> InFileActions)
@@ -175,7 +177,7 @@ private:
 	TWeakPtr<const ContentBrowserFileData::FFileActions> FileActions;
 };
 
-class CONTENTBROWSERFILEDATASOURCE_API FContentBrowserFileItemDataPayload_Duplication : public FContentBrowserFileItemDataPayload
+class FContentBrowserFileItemDataPayload_Duplication : public FContentBrowserFileItemDataPayload
 {
 public:
 	FContentBrowserFileItemDataPayload_Duplication(const FName InInternalPath, const FString& InFilename, TWeakPtr<const ContentBrowserFileData::FFileActions> InFileActions, const FString& InSourceFilename)
@@ -195,7 +197,7 @@ private:
 };
 
 
-class CONTENTBROWSERFILEDATASOURCE_API FContentBrowserFileItemDataPayload_Creation : public FContentBrowserFileItemDataPayload
+class FContentBrowserFileItemDataPayload_Creation : public FContentBrowserFileItemDataPayload
 {
 public:
 	FContentBrowserFileItemDataPayload_Creation(const FName InInternalPath, const FString& InFilename, TWeakPtr<const ContentBrowserFileData::FFileActions> InFileActions, FStructOnScope&& InCreationConfig)
@@ -212,3 +214,5 @@ public:
 private:
 	FStructOnScope CreationConfig;
 };
+
+#undef UE_API
diff --git a/Engine/Plugins/Editor/ContentBrowser/ContentBrowserFileDataSource/Source/ContentBrowserFileDataSource/Public/ContentBrowserFileDataSource.h b/Engine/Plugins/Editor/ContentBrowser/ContentBrowserFileDataSource/Source/ContentBrowserFileDataSource/Public/ContentBrowserFileDataSource.h
--- a/Engine/Plugins/Editor/ContentBrowser/ContentBrowserFileDataSource/Source/ContentBrowserFileDataSource/Public/ContentBrowserFileDataSource.h
+++ b/Engine/Plugins/Editor/ContentBrowser/ContentBrowserFileDataSource/Source/ContentBrowserFileDataSource/Public/ContentBrowserFileDataSource.h
@@ -8,13 +8,15 @@
 #include "ContentBrowserFileDataPayload.h"
 #include "ContentBrowserFileDataSource.generated.h"
 
+#define UE_API CONTENTBROWSERFILEDATASOURCE_API
+
 class UToolMenu;
 class FContentBrowserFileDataDiscovery;
 class IAssetTypeActions;
 struct FFileChangeData;
 
 USTRUCT()
-struct CONTENTBROWSERFILEDATASOURCE_API FContentBrowserCompiledFileDataFilter
+struct FContentBrowserCompiledFileDataFilter
 {
 	GENERATED_BODY()
 
@@ -27,91 +29,91 @@ public:
 	TArray<FString> FileExtensionsToInclude;
 };
 
-UCLASS()
-class CONTENTBROWSERFILEDATASOURCE_API UContentBrowserFileDataSource : public UContentBrowserDataSource
+UCLASS(MinimalAPI)
+class UContentBrowserFileDataSource : public UContentBrowserDataSource
 {
 	GENERATED_BODY()
 
 public:
-	void Initialize(const ContentBrowserFileData::FFileConfigData& InConfig, const bool InAutoRegister = true);
+	UE_API void Initialize(const ContentBrowserFileData::FFileConfigData& InConfig, const bool InAutoRegister = true);
 
-	virtual void Shutdown() override;
+	UE_API virtual void Shutdown() override;
 
-	void AddFileMount(const FName InFileMountPath, const FString& InFileMountDiskPath);
+	UE_API void AddFileMount(const FName InFileMountPath, const FString& InFileMountDiskPath);
 
-	void RemoveFileMount(const FName InFileMountPath);
+	UE_API void RemoveFileMount(const FName InFileMountPath);
 
-	bool HasFileMount(const FName InFileMountPath) const;
+	UE_API bool HasFileMount(const FName InFileMountPath) const;
 
-	virtual void Tick(const float InDeltaTime);
+	UE_API virtual void Tick(const float InDeltaTime);
 
-	virtual void CompileFilter(const FName InPath, const FContentBrowserDataFilter& InFilter, FContentBrowserDataCompiledFilter& OutCompiledFilter) override;
+	UE_API virtual void CompileFilter(const FName InPath, const FContentBrowserDataFilter& InFilter, FContentBrowserDataCompiledFilter& OutCompiledFilter) override;
 
-	virtual void EnumerateItemsMatchingFilter(const FContentBrowserDataCompiledFilter& InFilter, TFunctionRef<bool(FContentBrowserItemData&&)> InCallback) override;
+	UE_API virtual void EnumerateItemsMatchingFilter(const FContentBrowserDataCompiledFilter& InFilter, TFunctionRef<bool(FContentBrowserItemData&&)> InCallback) override;
 
-	virtual void EnumerateItemsAtPath(const FName InPath, const EContentBrowserItemTypeFilter InItemTypeFilter, TFunctionRef<bool(FContentBrowserItemData&&)> InCallback) override;
+	UE_API virtual void EnumerateItemsAtPath(const FName InPath, const EContentBrowserItemTypeFilter InItemTypeFilter, TFunctionRef<bool(FContentBrowserItemData&&)> InCallback) override;
 
-	virtual bool CanCreateFolder(const FName InPath, FText* OutErrorMsg) override;
+	UE_API virtual bool CanCreateFolder(const FName InPath, FText* OutErrorMsg) override;
 
-	virtual bool CreateFolder(const FName InPath, const TSharedPtr<IContentBrowserHideFolderIfEmptyFilter>& HideFolderIfEmptyFilter, FContentBrowserItemDataTemporaryContext& OutPendingItem) override;
+	UE_API virtual bool CreateFolder(const FName InPath, const TSharedPtr<IContentBrowserHideFolderIfEmptyFilter>& HideFolderIfEmptyFilter, FContentBrowserItemDataTemporaryContext& OutPendingItem) override;
 
-	virtual bool IsDiscoveringItems(FText* OutStatus = nullptr) override;
+	UE_API virtual bool IsDiscoveringItems(FText* OutStatus = nullptr) override;
 
-	virtual bool PrioritizeSearchPath(const FName InPath) override;
+	UE_API virtual bool PrioritizeSearchPath(const FName InPath) override;
 
-	virtual bool IsFolderVisible(const FName InPath, const EContentBrowserIsFolderVisibleFlags InFlags, const FContentBrowserFolderContentsFilter& InContentsFilter) override;
+	UE_API virtual bool IsFolderVisible(const FName InPath, const EContentBrowserIsFolderVisibleFlags InFlags, const FContentBrowserFolderContentsFilter& InContentsFilter) override;
 
-	virtual bool DoesItemPassFilter(const FContentBrowserItemData& InItem, const FContentBrowserDataCompiledFilter& InFilter) override;
+	UE_API virtual bool DoesItemPassFilter(const FContentBrowserItemData& InItem, const FContentBrowserDataCompiledFilter& InFilter) override;
 
-	virtual bool GetItemAttribute(const FContentBrowserItemData& InItem, const bool InIncludeMetaData, const FName InAttributeKey, FContentBrowserItemDataAttributeValue& OutAttributeValue) override;
+	UE_API virtual bool GetItemAttribute(const FContentBrowserItemData& InItem, const bool InIncludeMetaData, const FName InAttributeKey, FContentBrowserItemDataAttributeValue& OutAttributeValue) override;
 
-	virtual bool GetItemAttributes(const FContentBrowserItemData& InItem, const bool InIncludeMetaData, FContentBrowserItemDataAttributeValues& OutAttributeValues) override;
+	UE_API virtual bool GetItemAttributes(const FContentBrowserItemData& InItem, const bool InIncludeMetaData, FContentBrowserItemDataAttributeValues& OutAttributeValues) override;
 
-	virtual bool GetItemPhysicalPath(const FContentBrowserItemData& InItem, FString& OutDiskPath) override;
+	UE_API virtual bool GetItemPhysicalPath(const FContentBrowserItemData& InItem, FString& OutDiskPath) override;
 
-	virtual bool CanEditItem(const FContentBrowserItemData& InItem, FText* OutErrorMsg) override;
+	UE_API virtual bool CanEditItem(const FContentBrowserItemData& InItem, FText* OutErrorMsg) override;
 
-	virtual bool EditItem(const FContentBrowserItemData& InItem) override;
+	UE_API virtual bool EditItem(const FContentBrowserItemData& InItem) override;
 
-	virtual bool BulkEditItems(TArrayView<const FContentBrowserItemData> InItems) override;
+	UE_API virtual bool BulkEditItems(TArrayView<const FContentBrowserItemData> InItems) override;
 
-	virtual bool CanPreviewItem(const FContentBrowserItemData& InItem, FText* OutErrorMsg) override;
+	UE_API virtual bool CanPreviewItem(const FContentBrowserItemData& InItem, FText* OutErrorMsg) override;
 
-	virtual bool PreviewItem(const FContentBrowserItemData& InItem) override;
+	UE_API virtual bool PreviewItem(const FContentBrowserItemData& InItem) override;
 
-	virtual bool BulkPreviewItems(TArrayView<const FContentBrowserItemData> InItems) override;
+	UE_API virtual bool BulkPreviewItems(TArrayView<const FContentBrowserItemData> InItems) override;
 
-	virtual bool CanDuplicateItem(const FContentBrowserItemData& InItem, FText* OutErrorMsg) override;
+	UE_API virtual bool CanDuplicateItem(const FContentBrowserItemData& InItem, FText* OutErrorMsg) override;
 
-	virtual bool DuplicateItem(const FContentBrowserItemData& InItem, FContentBrowserItemDataTemporaryContext& OutPendingItem) override;
+	UE_API virtual bool DuplicateItem(const FContentBrowserItemData& InItem, FContentBrowserItemDataTemporaryContext& OutPendingItem) override;
 
-	virtual bool BulkDuplicateItems(TArrayView<const FContentBrowserItemData> InItems, TArray<FContentBrowserItemData>& OutNewItems) override;
+	UE_API virtual bool BulkDuplicateItems(TArrayView<const FContentBrowserItemData> InItems, TArray<FContentBrowserItemData>& OutNewItems) override;
 
-	virtual bool CanDeleteItem(const FContentBrowserItemData& InItem, FText* OutErrorMsg) override;
+	UE_API virtual bool CanDeleteItem(const FContentBrowserItemData& InItem, FText* OutErrorMsg) override;
 
-	virtual bool DeleteItem(const FContentBrowserItemData& InItem) override;
+	UE_API virtual bool DeleteItem(const FContentBrowserItemData& InItem) override;
 
-	virtual bool BulkDeleteItems(TArrayView<const FContentBrowserItemData> InItems) override;
+	UE_API virtual bool BulkDeleteItems(TArrayView<const FContentBrowserItemData> InItems) override;
 
-	virtual bool CanRenameItem(const FContentBrowserItemData& InItem, const FString* InNewName, const IContentBrowserHideFolderIfEmptyFilter* HideFolderIfEmptyFilter, FText* OutErrorMsg) override;
+	UE_API virtual bool CanRenameItem(const FContentBrowserItemData& InItem, const FString* InNewName, const IContentBrowserHideFolderIfEmptyFilter* HideFolderIfEmptyFilter, FText* OutErrorMsg) override;
 
-	virtual bool RenameItem(const FContentBrowserItemData& InItem, const FString& InNewName, FContentBrowserItemData& OutNewItem) override;
+	UE_API virtual bool RenameItem(const FContentBrowserItemData& InItem, const FString& InNewName, FContentBrowserItemData& OutNewItem) override;
 
-	virtual bool CanCopyItem(const FContentBrowserItemData& InItem, const FName InDestPath, FText* OutErrorMsg) override;
+	UE_API virtual bool CanCopyItem(const FContentBrowserItemData& InItem, const FName InDestPath, FText* OutErrorMsg) override;
 
-	virtual bool CopyItem(const FContentBrowserItemData& InItem, const FName InDestPath) override;
+	UE_API virtual bool CopyItem(const FContentBrowserItemData& InItem, const FName InDestPath) override;
 
-	virtual bool BulkCopyItems(TArrayView<const FContentBrowserItemData> InItems, const FName InDestPath) override;
+	UE_API virtual bool BulkCopyItems(TArrayView<const FContentBrowserItemData> InItems, const FName InDestPath) override;
 
-	virtual bool CanMoveItem(const FContentBrowserItemData& InItem, const FName InDestPath, FText* OutErrorMsg) override;
+	UE_API virtual bool CanMoveItem(const FContentBrowserItemData& InItem, const FName InDestPath, FText* OutErrorMsg) override;
 
-	virtual bool MoveItem(const FContentBrowserItemData& InItem, const FName InDestPath) override;
+	UE_API virtual bool MoveItem(const FContentBrowserItemData& InItem, const FName InDestPath) override;
 
-	virtual bool BulkMoveItems(TArrayView<const FContentBrowserItemData> InItems, const FName InDestPath) override;
+	UE_API virtual bool BulkMoveItems(TArrayView<const FContentBrowserItemData> InItems, const FName InDestPath) override;
 
-	virtual bool UpdateThumbnail(const FContentBrowserItemData& InItem, FAssetThumbnail& InThumbnail) override;
+	UE_API virtual bool UpdateThumbnail(const FContentBrowserItemData& InItem, FAssetThumbnail& InThumbnail) override;
 
-	virtual void BuildRootPathVirtualTree() override;
+	UE_API virtual void BuildRootPathVirtualTree() override;
 
 protected:
 
@@ -134,52 +136,52 @@ protected:
 		TSet<FName> ChildItems;
 	};
 
-	FContentBrowserItemData CreateFolderItem(const FName InInternalPath, const FString& InFilename);
+	UE_API FContentBrowserItemData CreateFolderItem(const FName InInternalPath, const FString& InFilename);
 
-	FContentBrowserItemData CreateFileItem(const FName InInternalPath, const FString& InFilename);
+	UE_API FContentBrowserItemData CreateFileItem(const FName InInternalPath, const FString& InFilename);
 
-	FContentBrowserItemData CreateItemFromDiscovered(const FName InInternalPath, const FDiscoveredItem& InDiscoveredItem);
+	UE_API FContentBrowserItemData CreateItemFromDiscovered(const FName InInternalPath, const FDiscoveredItem& InDiscoveredItem);
 
-	TSharedPtr<const FContentBrowserFolderItemDataPayload> GetFolderItemPayload(const FContentBrowserItemData& InItem) const;
+	UE_API TSharedPtr<const FContentBrowserFolderItemDataPayload> GetFolderItemPayload(const FContentBrowserItemData& InItem) const;
 
-	TSharedPtr<const FContentBrowserFileItemDataPayload> GetFileItemPayload(const FContentBrowserItemData& InItem) const;
+	UE_API TSharedPtr<const FContentBrowserFileItemDataPayload> GetFileItemPayload(const FContentBrowserItemData& InItem) const;
 
-	bool HideFolderIfEmpty(const IContentBrowserHideFolderIfEmptyFilter& HideFolderIfEmptyFilter, FName Path, FStringView PathString) const;
+	UE_API bool HideFolderIfEmpty(const IContentBrowserHideFolderIfEmptyFilter& HideFolderIfEmptyFilter, FName Path, FStringView PathString) const;
 
-	bool IsKnownFileMount(const FName InMountPath, FString* OutDiskPath = nullptr) const;
+	UE_API bool IsKnownFileMount(const FName InMountPath, FString* OutDiskPath = nullptr) const;
 
-	bool IsRootFileMount(const FName InMountPath, FString* OutDiskPath = nullptr) const;
+	UE_API bool IsRootFileMount(const FName InMountPath, FString* OutDiskPath = nullptr) const;
 
-	void AddDiscoveredItem(FDiscoveredItem::EType InItemType, const FString& InMountPath, const FString& InDiskPath, const bool bIsRootPath);
-	void AddDiscoveredItemImpl(FDiscoveredItem::EType InItemType, const FString& InMountPath, const FString& InDiskPath, const FName InChildMountPathName, const bool bIsRootPath);
+	UE_API void AddDiscoveredItem(FDiscoveredItem::EType InItemType, const FString& InMountPath, const FString& InDiskPath, const bool bIsRootPath);
+	UE_API void AddDiscoveredItemImpl(FDiscoveredItem::EType InItemType, const FString& InMountPath, const FString& InDiskPath, const FName InChildMountPathName, const bool bIsRootPath);
 
-	void RemoveDiscoveredItem(const FString& InMountPath);
-	void RemoveDiscoveredItem(const FName InMountPath);
-	void RemoveDiscoveredItemImpl(const FName InMountPath, const bool bParentIsOrphan);
+	UE_API void RemoveDiscoveredItem(const FString& InMountPath);
+	UE_API void RemoveDiscoveredItem(const FName InMountPath);
+	UE_API void RemoveDiscoveredItemImpl(const FName InMountPath, const bool bParentIsOrphan);
 
-	void OnPathPopulated(const FName InPath);
-	void OnPathPopulated(const FStringView InPath);
+	UE_API void OnPathPopulated(const FName InPath);
+	UE_API void OnPathPopulated(const FStringView InPath);
 
-	void OnAlwaysShowPath(const FName InPath);
-	void OnAlwaysShowPath(const FStringView InPath);
+	UE_API void OnAlwaysShowPath(const FName InPath);
+	UE_API void OnAlwaysShowPath(const FStringView InPath);
 
 	/** Called when a file in a watched directory changes on disk */
-	void OnDirectoryChanged(const TArray<FFileChangeData>& InFileChanges, const FString InFileMountPath, const FString InFileMountDiskPath);
+	UE_API void OnDirectoryChanged(const TArray<FFileChangeData>& InFileChanges, const FString InFileMountPath, const FString InFileMountDiskPath);
 
-	void PopulateAddNewContextMenu(UToolMenu* InMenu);
+	UE_API void PopulateAddNewContextMenu(UToolMenu* InMenu);
 
-	void OnNewFileRequested(const FName InDestFolderPath, const FString InDestFolder, TSharedRef<const ContentBrowserFileData::FFileActions> InFileActions, UContentBrowserDataMenuContext_AddNewMenu::FOnBeginItemCreation InOnBeginItemCreation);
+	UE_API void OnNewFileRequested(const FName InDestFolderPath, const FString InDestFolder, TSharedRef<const ContentBrowserFileData::FFileActions> InFileActions, UContentBrowserDataMenuContext_AddNewMenu::FOnBeginItemCreation InOnBeginItemCreation);
 
-	bool OnValidateItemName(const FContentBrowserItemData& InItem, const FString& InProposedName, FText* OutErrorMsg);
+	UE_API bool OnValidateItemName(const FContentBrowserItemData& InItem, const FString& InProposedName, FText* OutErrorMsg);
 
-	FContentBrowserItemData OnFinalizeCreateFolder(const FContentBrowserItemData& InItemData, const FString& InProposedName, FText* OutErrorMsg);
+	UE_API FContentBrowserItemData OnFinalizeCreateFolder(const FContentBrowserItemData& InItemData, const FString& InProposedName, FText* OutErrorMsg);
 
-	FContentBrowserItemData OnFinalizeCreateFile(const FContentBrowserItemData& InItemData, const FString& InProposedName, FText* OutErrorMsg);
+	UE_API FContentBrowserItemData OnFinalizeCreateFile(const FContentBrowserItemData& InItemData, const FString& InProposedName, FText* OutErrorMsg);
 
-	FContentBrowserItemData OnFinalizeDuplicateFile(const FContentBrowserItemData& InItemData, const FString& InProposedName, FText* OutErrorMsg);
+	UE_API FContentBrowserItemData OnFinalizeDuplicateFile(const FContentBrowserItemData& InItemData, const FString& InProposedName, FText* OutErrorMsg);
 
-	bool PassesFilters(const FName InPath, const int32 InFolderDepthChecked, const FContentBrowserCompiledFileDataFilter& InFileDataFilter) const;
-	static bool PassesFilters(const FStringView InPath, const FDiscoveredItem& InDiscoveredItem, const int32 InFolderDepthChecked, const FContentBrowserCompiledFileDataFilter& InFileDataFilter);
+	UE_API bool PassesFilters(const FName InPath, const int32 InFolderDepthChecked, const FContentBrowserCompiledFileDataFilter& InFileDataFilter) const;
+	static UE_API bool PassesFilters(const FStringView InPath, const FDiscoveredItem& InDiscoveredItem, const int32 InFolderDepthChecked, const FContentBrowserCompiledFileDataFilter& InFileDataFilter);
 
 	ContentBrowserFileData::FFileConfigData Config;
 
@@ -203,3 +205,5 @@ protected:
 	 */
 	TSet<FString> EmptyFolders;
 };
+
+#undef UE_API
diff --git a/Engine/Plugins/Editor/CryptoKeys/Source/CryptoKeys/Classes/CryptoKeysCommandlet.h b/Engine/Plugins/Editor/CryptoKeys/Source/CryptoKeys/Classes/CryptoKeysCommandlet.h
--- a/Engine/Plugins/Editor/CryptoKeys/Source/CryptoKeys/Classes/CryptoKeysCommandlet.h
+++ b/Engine/Plugins/Editor/CryptoKeys/Source/CryptoKeys/Classes/CryptoKeysCommandlet.h
@@ -5,15 +5,19 @@
 #include "Commandlets/Commandlet.h"
 #include "CryptoKeysCommandlet.generated.h"
 
+#define UE_API CRYPTOKEYS_API
+
 /**
 * Commandlet used to configure project encryption settings
 */
-UCLASS()
-class CRYPTOKEYS_API UCryptoKeysCommandlet : public UCommandlet
+UCLASS(MinimalAPI)
+class UCryptoKeysCommandlet : public UCommandlet
 {
 	GENERATED_UCLASS_BODY()
 
 		//~ Begin UCommandlet Interface
-		virtual int32 Main(const FString& Params) override;
+		UE_API virtual int32 Main(const FString& Params) override;
 	//~ End UCommandlet Interface
 };
+
+#undef UE_API
diff --git a/Engine/Plugins/Editor/CryptoKeys/Source/CryptoKeys/Classes/CryptoKeysSettings.h b/Engine/Plugins/Editor/CryptoKeys/Source/CryptoKeys/Classes/CryptoKeysSettings.h
--- a/Engine/Plugins/Editor/CryptoKeys/Source/CryptoKeys/Classes/CryptoKeysSettings.h
+++ b/Engine/Plugins/Editor/CryptoKeys/Source/CryptoKeys/Classes/CryptoKeysSettings.h
@@ -4,6 +4,8 @@
 
 #include "CryptoKeysSettings.generated.h"
 
+#define UE_API CRYPTOKEYS_API
+
 
 /**
 * UStruct representing a named encryption key
@@ -26,17 +28,17 @@ struct FCryptoEncryptionKey
 /**
 * Implements the settings for imported Paper2D assets, such as sprite sheet textures.
 */
-UCLASS(config = Crypto, defaultconfig)
-class CRYPTOKEYS_API UCryptoKeysSettings : public UObject
+UCLASS(MinimalAPI, config = Crypto, defaultconfig)
+class UCryptoKeysSettings : public UObject
 {
 	GENERATED_BODY()
 
 public:
 
-	UCryptoKeysSettings();
+	UE_API UCryptoKeysSettings();
 
 	//~ Begin UObject Interface
-	virtual void PostEditChangeProperty(struct FPropertyChangedEvent& PropertyChangedEvent);
+	UE_API virtual void PostEditChangeProperty(struct FPropertyChangedEvent& PropertyChangedEvent);
 	//~ End UObject Interface
 
 	bool IsEncryptionEnabled() const
@@ -89,3 +91,5 @@ public:
 	UPROPERTY(config, EditAnywhere, Category = Signing)
 	bool bEnablePakSigning;
 };
+
+#undef UE_API
diff --git a/Engine/Plugins/Editor/DataValidation/Source/DataValidation/Private/PackageFileValidator.h b/Engine/Plugins/Editor/DataValidation/Source/DataValidation/Private/PackageFileValidator.h
--- a/Engine/Plugins/Editor/DataValidation/Source/DataValidation/Private/PackageFileValidator.h
+++ b/Engine/Plugins/Editor/DataValidation/Source/DataValidation/Private/PackageFileValidator.h
@@ -6,6 +6,8 @@
 
 #include "PackageFileValidator.generated.h"
 
+#define UE_API DATAVALIDATION_API
+
 struct FPackageFileSummary;
 
 /**
@@ -15,23 +17,25 @@ struct FPackageFileSummary;
  * To disable validation of payload hashes (which is much slower than the rest of the validation) set 
  * ini:Editor:[/Script/DataValidation.PackageFileValidator]:bValidatePayloadHashes=true
  */
-UCLASS()
-class DATAVALIDATION_API UPackageFileValidator : public UEditorValidatorBase
+UCLASS(MinimalAPI)
+class UPackageFileValidator : public UEditorValidatorBase
 {
 	GENERATED_BODY()
 
 protected:
 
-	virtual bool CanValidateAsset_Implementation(const FAssetData& AssetData, UObject* Asset, FDataValidationContext& Context) const override;
-	virtual EDataValidationResult ValidateLoadedAsset_Implementation(const FAssetData& AssetData, UObject* Asset, FDataValidationContext& Context) override;
+	UE_API virtual bool CanValidateAsset_Implementation(const FAssetData& AssetData, UObject* Asset, FDataValidationContext& Context) const override;
+	UE_API virtual EDataValidationResult ValidateLoadedAsset_Implementation(const FAssetData& AssetData, UObject* Asset, FDataValidationContext& Context) override;
 
 private:
 
-	bool ValidatePackageSummary(FName PackageName, FArchive& Ar, FDataValidationContext& Context, FPackageFileSummary& OutSummary);
-	bool ValidatePackageTrailer(FName PackageName, FArchive& Ar, FDataValidationContext& Context);
+	UE_API bool ValidatePackageSummary(FName PackageName, FArchive& Ar, FDataValidationContext& Context, FPackageFileSummary& OutSummary);
+	UE_API bool ValidatePackageTrailer(FName PackageName, FArchive& Ar, FDataValidationContext& Context);
 
-	bool TryResolvePackagePath(FName PackageName, FPackagePath& OutPackagePath) const;
+	UE_API bool TryResolvePackagePath(FName PackageName, FPackagePath& OutPackagePath) const;
 
 	UPROPERTY(Config)
 	bool bValidatePayloadHashes = true;
 };
+
+#undef UE_API
diff --git a/Engine/Plugins/Editor/DataValidation/Source/DataValidation/Private/WorldPartitionChangelistValidator.h b/Engine/Plugins/Editor/DataValidation/Source/DataValidation/Private/WorldPartitionChangelistValidator.h
--- a/Engine/Plugins/Editor/DataValidation/Source/DataValidation/Private/WorldPartitionChangelistValidator.h
+++ b/Engine/Plugins/Editor/DataValidation/Source/DataValidation/Private/WorldPartitionChangelistValidator.h
@@ -6,12 +6,14 @@
 
 #include "WorldPartitionChangelistValidator.generated.h"
 
+#define UE_API DATAVALIDATION_API
+
 class UDataValidationChangelist;
 class UActorDescContainerInstance;
 class FWorldPartitionActorDesc;
 
-UCLASS()
-class DATAVALIDATION_API UWorldPartitionChangelistValidator : public UEditorValidatorBase, public IStreamingGenerationErrorHandler
+UCLASS(MinimalAPI)
+class UWorldPartitionChangelistValidator : public UEditorValidatorBase, public IStreamingGenerationErrorHandler
 {
 	GENERATED_BODY()
 
@@ -23,32 +25,34 @@ protected:
 	UObject* CurrentAsset = nullptr;
 	bool bSubmittingWorldDataLayers = false;
 
-	virtual bool CanValidateAsset_Implementation(const FAssetData& InAssetData, UObject* InAsset, FDataValidationContext& InContext) const override;
-	virtual EDataValidationResult ValidateLoadedAsset_Implementation(const FAssetData& InAssetData, UObject* InAsset, FDataValidationContext& InContext) override;
+	UE_API virtual bool CanValidateAsset_Implementation(const FAssetData& InAssetData, UObject* InAsset, FDataValidationContext& InContext) const override;
+	UE_API virtual EDataValidationResult ValidateLoadedAsset_Implementation(const FAssetData& InAssetData, UObject* InAsset, FDataValidationContext& InContext) override;
 
-	void ValidateActorsAndDataLayersFromChangeList(UDataValidationChangelist* Changelist);
+	UE_API void ValidateActorsAndDataLayersFromChangeList(UDataValidationChangelist* Changelist);
 
 	// Return true if this ActorDescView is pertinent to the current ChangeList
-	bool Filter(const IWorldPartitionActorDescInstanceView& ActorDescView);
+	UE_API bool Filter(const IWorldPartitionActorDescInstanceView& ActorDescView);
 
 	// Return true if this UDataLayerInstance is pertinent to the current ChangeList
-	bool Filter(const UDataLayerInstance* DataLayerInstance);
+	UE_API bool Filter(const UDataLayerInstance* DataLayerInstance);
 
 	// IStreamingGenerationErrorHandler Interface methods
-	virtual void OnInvalidRuntimeGrid(const IWorldPartitionActorDescInstanceView& ActorDescView, FName GridName) override;
-	virtual void OnInvalidReference(const IWorldPartitionActorDescInstanceView& ActorDescView, const FGuid& ReferenceGuid, IWorldPartitionActorDescInstanceView* ReferenceActorDescView) override;
-	virtual void OnInvalidReferenceGridPlacement(const IWorldPartitionActorDescInstanceView& ActorDescView, const IWorldPartitionActorDescInstanceView& ReferenceActorDescView) override;
-	virtual void OnInvalidReferenceDataLayers(const IWorldPartitionActorDescInstanceView& ActorDescView, const IWorldPartitionActorDescInstanceView& ReferenceActorDescView, EDataLayerInvalidReason Reason) override;
-	virtual void OnInvalidWorldReference(const IWorldPartitionActorDescInstanceView& ActorDescView, EWorldReferenceInvalidReason Reason) override;
-	virtual void OnInvalidReferenceRuntimeGrid(const IWorldPartitionActorDescInstanceView& ActorDescView, const IWorldPartitionActorDescInstanceView& ReferenceActorDescView) override;
-	virtual void OnDataLayersLoadFilterMismatch(const IWorldPartitionActorDescInstanceView& ActorDescView) override;
-	virtual void OnInvalidReferenceDataLayerAsset(const UDataLayerInstanceWithAsset* DataLayerInstance) override;
-	virtual void OnInvalidDataLayerAssetType(const UDataLayerInstanceWithAsset* DataLayerInstance, const UDataLayerAsset* DataLayerAsset) override;
-	virtual void OnDataLayerHierarchyTypeMismatch(const UDataLayerInstance* DataLayerInstance, const UDataLayerInstance* Parent, EDataLayerHierarchyInvalidReason Reason) override;
-	virtual void OnInvalidWorldDataLayersReference(const AWorldDataLayers* WorldDataLayers, const UDataLayerInstance* DataLayerInstance, const FText& Reason) override;
-	virtual void OnDataLayerAssetConflict(const UDataLayerInstanceWithAsset* DataLayerInstance, const UDataLayerInstanceWithAsset* ConflictingDataLayerInstance) override;
-	virtual void OnActorNeedsResave(const IWorldPartitionActorDescInstanceView& ActorDescView) override;
-	virtual void OnLevelInstanceInvalidWorldAsset(const IWorldPartitionActorDescInstanceView& ActorDescView, FName WorldAsset, ELevelInstanceInvalidReason Reason) override;
-	virtual void OnInvalidActorFilterReference(const IWorldPartitionActorDescInstanceView& ActorDescView, const IWorldPartitionActorDescInstanceView& ReferenceActorDescView) override;
-	virtual void OnInvalidHLODLayer(const IWorldPartitionActorDescInstanceView& ActorDescView) override;
+	UE_API virtual void OnInvalidRuntimeGrid(const IWorldPartitionActorDescInstanceView& ActorDescView, FName GridName) override;
+	UE_API virtual void OnInvalidReference(const IWorldPartitionActorDescInstanceView& ActorDescView, const FGuid& ReferenceGuid, IWorldPartitionActorDescInstanceView* ReferenceActorDescView) override;
+	UE_API virtual void OnInvalidReferenceGridPlacement(const IWorldPartitionActorDescInstanceView& ActorDescView, const IWorldPartitionActorDescInstanceView& ReferenceActorDescView) override;
+	UE_API virtual void OnInvalidReferenceDataLayers(const IWorldPartitionActorDescInstanceView& ActorDescView, const IWorldPartitionActorDescInstanceView& ReferenceActorDescView, EDataLayerInvalidReason Reason) override;
+	UE_API virtual void OnInvalidWorldReference(const IWorldPartitionActorDescInstanceView& ActorDescView, EWorldReferenceInvalidReason Reason) override;
+	UE_API virtual void OnInvalidReferenceRuntimeGrid(const IWorldPartitionActorDescInstanceView& ActorDescView, const IWorldPartitionActorDescInstanceView& ReferenceActorDescView) override;
+	UE_API virtual void OnDataLayersLoadFilterMismatch(const IWorldPartitionActorDescInstanceView& ActorDescView) override;
+	UE_API virtual void OnInvalidReferenceDataLayerAsset(const UDataLayerInstanceWithAsset* DataLayerInstance) override;
+	UE_API virtual void OnInvalidDataLayerAssetType(const UDataLayerInstanceWithAsset* DataLayerInstance, const UDataLayerAsset* DataLayerAsset) override;
+	UE_API virtual void OnDataLayerHierarchyTypeMismatch(const UDataLayerInstance* DataLayerInstance, const UDataLayerInstance* Parent, EDataLayerHierarchyInvalidReason Reason) override;
+	UE_API virtual void OnInvalidWorldDataLayersReference(const AWorldDataLayers* WorldDataLayers, const UDataLayerInstance* DataLayerInstance, const FText& Reason) override;
+	UE_API virtual void OnDataLayerAssetConflict(const UDataLayerInstanceWithAsset* DataLayerInstance, const UDataLayerInstanceWithAsset* ConflictingDataLayerInstance) override;
+	UE_API virtual void OnActorNeedsResave(const IWorldPartitionActorDescInstanceView& ActorDescView) override;
+	UE_API virtual void OnLevelInstanceInvalidWorldAsset(const IWorldPartitionActorDescInstanceView& ActorDescView, FName WorldAsset, ELevelInstanceInvalidReason Reason) override;
+	UE_API virtual void OnInvalidActorFilterReference(const IWorldPartitionActorDescInstanceView& ActorDescView, const IWorldPartitionActorDescInstanceView& ReferenceActorDescView) override;
+	UE_API virtual void OnInvalidHLODLayer(const IWorldPartitionActorDescInstanceView& ActorDescView) override;
 };
+
+#undef UE_API
diff --git a/Engine/Plugins/Editor/DataValidation/Source/DataValidation/Public/DataValidationChangelist.h b/Engine/Plugins/Editor/DataValidation/Source/DataValidation/Public/DataValidationChangelist.h
--- a/Engine/Plugins/Editor/DataValidation/Source/DataValidation/Public/DataValidationChangelist.h
+++ b/Engine/Plugins/Editor/DataValidation/Source/DataValidation/Public/DataValidationChangelist.h
@@ -9,11 +9,13 @@
 
 #include "DataValidationChangelist.generated.h"
 
+#define UE_API DATAVALIDATION_API
+
 /**
  * Changelist abstraction to allow changelist-level data validation
  */
-UCLASS(config = Editor)
-class DATAVALIDATION_API UDataValidationChangelist : public UObject
+UCLASS(MinimalAPI, config = Editor)
+class UDataValidationChangelist : public UObject
 {
 	GENERATED_BODY()
 
@@ -21,17 +23,17 @@ public:
 	/** Default constructor with nothing */
 	UDataValidationChangelist() = default;
 
-	virtual EDataValidationResult IsDataValid(class FDataValidationContext& Context) const override;
+	UE_API virtual EDataValidationResult IsDataValid(class FDataValidationContext& Context) const override;
 
 	/** Initializes from a list of file states as a pseudo-changelist */
-	void Initialize(TConstArrayView<FSourceControlStateRef> InFileStates);
+	UE_API void Initialize(TConstArrayView<FSourceControlStateRef> InFileStates);
 	/** Initializes from a changelist reference, querying the state from the provider */
-	void Initialize(FSourceControlChangelistPtr InChangelist);
+	UE_API void Initialize(FSourceControlChangelistPtr InChangelist);
 	/** Initializes from an already-queried changelist state */
-	void Initialize(FSourceControlChangelistStateRef InChangelistState);
+	UE_API void Initialize(FSourceControlChangelistStateRef InChangelistState);
 
-	static void GatherDependencies(const FName& InPackageName, TSet<TPair<FName, FName>>& OutDependencies);
-	static FString GetPrettyPackageName(const FName& InPackageName);
+	static UE_API void GatherDependencies(const FName& InPackageName, TSet<TPair<FName, FName>>& OutDependencies);
+	static UE_API FString GetPrettyPackageName(const FName& InPackageName);
 
 	/** Changelist to validate - may be null if this was constructed from a list of files */
 	FSourceControlChangelistPtr Changelist;
@@ -46,3 +48,5 @@ public:
 	
 	FText Description;
 };
+
+#undef UE_API
diff --git a/Engine/Plugins/Editor/DataValidation/Source/DataValidation/Public/DataValidationCommandlet.h b/Engine/Plugins/Editor/DataValidation/Source/DataValidation/Public/DataValidationCommandlet.h
--- a/Engine/Plugins/Editor/DataValidation/Source/DataValidation/Public/DataValidationCommandlet.h
+++ b/Engine/Plugins/Editor/DataValidation/Source/DataValidation/Public/DataValidationCommandlet.h
@@ -4,8 +4,10 @@
 #include "Commandlets/Commandlet.h"
 #include "DataValidationCommandlet.generated.h"
 
-UCLASS(CustomConstructor)
-class DATAVALIDATION_API UDataValidationCommandlet : public UCommandlet
+#define UE_API DATAVALIDATION_API
+
+UCLASS(MinimalAPI, CustomConstructor)
+class UDataValidationCommandlet : public UCommandlet
 {
 	GENERATED_UCLASS_BODY()
 
@@ -17,9 +19,11 @@ public:
 	}
 
 	// Begin UCommandlet Interface
-	virtual int32 Main(const FString& FullCommandLine) override;
+	UE_API virtual int32 Main(const FString& FullCommandLine) override;
 	// End UCommandlet Interface
 
 	// do the validation without creating a commandlet
-	static bool ValidateData(const FString& FullCommandLine);
+	static UE_API bool ValidateData(const FString& FullCommandLine);
 };
+
+#undef UE_API
diff --git a/Engine/Plugins/Editor/DataValidation/Source/DataValidation/Public/DataValidationFixers.h b/Engine/Plugins/Editor/DataValidation/Source/DataValidation/Public/DataValidationFixers.h
--- a/Engine/Plugins/Editor/DataValidation/Source/DataValidation/Public/DataValidationFixers.h
+++ b/Engine/Plugins/Editor/DataValidation/Source/DataValidation/Public/DataValidationFixers.h
@@ -7,6 +7,8 @@
 #include "Templates/SharedPointer.h"
 #include "UObject/WeakObjectPtr.h"
 
+#define UE_API DATAVALIDATION_API
+
 namespace UE::DataValidation
 {
 
@@ -64,68 +66,68 @@ TSharedRef<TLambdaFixer<FFixAlwaysApplicable, ApplyFixT>> MakeFix(ApplyFixT&& Ap
  * IFixer which wraps another fixer and makes it single use only.
  * All fixers which are not idempotent will want to use this.
  */
-struct DATAVALIDATION_API FSingleUseFixer : IFixer
+struct FSingleUseFixer : IFixer
 {
 	TSharedPtr<IFixer> Inner;
 	TSet<int32> UsedFixes;
 
-	virtual EFixApplicability GetApplicability(int32 FixIndex) const override;
-	virtual FFixResult ApplyFix(int32 FixIndex) override;
+	UE_API virtual EFixApplicability GetApplicability(int32 FixIndex) const override;
+	UE_API virtual FFixResult ApplyFix(int32 FixIndex) override;
 
-	static TSharedRef<FSingleUseFixer> Create(TSharedRef<IFixer> Inner);
+	static UE_API TSharedRef<FSingleUseFixer> Create(TSharedRef<IFixer> Inner);
 };
 
 /** IFixer which wraps another fixer in a set of dependencies. */
-struct DATAVALIDATION_API FObjectSetDependentFixer : IFixer
+struct FObjectSetDependentFixer : IFixer
 {
 	TSharedPtr<IFixer> Inner;
 	TArray<TWeakObjectPtr<>> Dependencies;
 
-	virtual EFixApplicability GetApplicability(int32 FixIndex) const override;
-	virtual FFixResult ApplyFix(int32 FixIndex) override;
+	UE_API virtual EFixApplicability GetApplicability(int32 FixIndex) const override;
+	UE_API virtual FFixResult ApplyFix(int32 FixIndex) override;
 
-	static TSharedRef<FObjectSetDependentFixer> Create(TSharedRef<IFixer> Inner, TArray<TWeakObjectPtr<>> Dependencies);
+	static UE_API TSharedRef<FObjectSetDependentFixer> Create(TSharedRef<IFixer> Inner, TArray<TWeakObjectPtr<>> Dependencies);
 };
 
 /** IFixer which wraps another fixer and makes it auto-save the asset after the fix is applied. */
-struct DATAVALIDATION_API FAutoSavingFixer : IFixer
+struct FAutoSavingFixer : IFixer
 {
 	TSharedPtr<IFixer> Inner;
 
-	virtual EFixApplicability GetApplicability(int32 FixIndex) const override;
-	virtual FFixResult ApplyFix(int32 FixIndex) override;
+	UE_API virtual EFixApplicability GetApplicability(int32 FixIndex) const override;
+	UE_API virtual FFixResult ApplyFix(int32 FixIndex) override;
 
-	static TSharedRef<FAutoSavingFixer> Create(TSharedRef<IFixer> Inner);
+	static UE_API TSharedRef<FAutoSavingFixer> Create(TSharedRef<IFixer> Inner);
 };
 
 /** IFixer which wraps another fixer and makes it validate any assets touched by the fix. */
-struct DATAVALIDATION_API FValidatingFixer : IFixer
+struct FValidatingFixer : IFixer
 {
 	TSharedPtr<IFixer> Inner;
 
-	virtual EFixApplicability GetApplicability(int32 FixIndex) const override;
-	virtual FFixResult ApplyFix(int32 FixIndex) override;
+	UE_API virtual EFixApplicability GetApplicability(int32 FixIndex) const override;
+	UE_API virtual FFixResult ApplyFix(int32 FixIndex) override;
 
-	static TSharedRef<FValidatingFixer> Create(TSharedRef<IFixer> Inner);
+	static UE_API TSharedRef<FValidatingFixer> Create(TSharedRef<IFixer> Inner);
 };
 
 /**
  * IFixer made out of many smaller fixers. Only one of the fixes in the set may be applied.
  * Once applied, the remaining fixers become non-applicable.
  */
-struct DATAVALIDATION_API FMutuallyExclusiveFixSet
+struct FMutuallyExclusiveFixSet
 {
-	FMutuallyExclusiveFixSet();
+	UE_API FMutuallyExclusiveFixSet();
 
 	/** Add a fixer to the set. */
-	void Add(const FText& Label, TSharedRef<IFixer> Inner);
+	UE_API void Add(const FText& Label, TSharedRef<IFixer> Inner);
 	/**
 	 * Transform all fixes.
 	 * The callback will be called for each fix and should return a new, transformed fix.
 	 */
-	void Transform(const TFunctionRef<TSharedRef<IFixer>(TSharedRef<IFixer>)>& Callback);
+	UE_API void Transform(const TFunctionRef<TSharedRef<IFixer>(TSharedRef<IFixer>)>& Callback);
 	/** Generate FFixTokens from the set. These tokens can then be added to a message. */
-	void CreateTokens(const TFunctionRef<void(TSharedRef<FFixToken>)>& Callback) const;
+	UE_API void CreateTokens(const TFunctionRef<void(TSharedRef<FFixToken>)>& Callback) const;
 
 private:
 	struct FSharedData
@@ -159,3 +161,4 @@ private:
 
 }
 
+#undef UE_API
diff --git a/Engine/Plugins/Editor/DataValidation/Source/DataValidation/Public/DataValidationModule.h b/Engine/Plugins/Editor/DataValidation/Source/DataValidation/Public/DataValidationModule.h
--- a/Engine/Plugins/Editor/DataValidation/Source/DataValidation/Public/DataValidationModule.h
+++ b/Engine/Plugins/Editor/DataValidation/Source/DataValidation/Public/DataValidationModule.h
@@ -11,7 +11,7 @@ struct FAssetData;
 /**
  * The public interface to this module
  */
-class DATAVALIDATION_API IDataValidationModule : public IModuleInterface
+class IDataValidationModule : public IModuleInterface
 {
 
 public:
diff --git a/Engine/Plugins/Editor/DataValidation/Source/DataValidation/Public/DataValidationSettings.h b/Engine/Plugins/Editor/DataValidation/Source/DataValidation/Public/DataValidationSettings.h
--- a/Engine/Plugins/Editor/DataValidation/Source/DataValidation/Public/DataValidationSettings.h
+++ b/Engine/Plugins/Editor/DataValidation/Source/DataValidation/Public/DataValidationSettings.h
@@ -7,17 +7,19 @@
 
 #include "DataValidationSettings.generated.h"
 
+#define UE_API DATAVALIDATION_API
+
 /**
  * Project-wide settings for data validation
  */
-UCLASS(config = Editor)
-class DATAVALIDATION_API UDataValidationSettings : public UDeveloperSettings
+UCLASS(MinimalAPI, config = Editor)
+class UDataValidationSettings : public UDeveloperSettings
 {
 	GENERATED_BODY()
 
 public:
 	/** Default constructor that sets up CDO properties */
-	UDataValidationSettings();
+	UE_API UDataValidationSettings();
 
 	/** Whether or not to validate assets on save */
 	UPROPERTY(EditAnywhere, Config, Category="Data Validation")
@@ -36,3 +38,5 @@ public:
 	UPROPERTY(EditAnywhere, Config, Category="Data Validation", meta=(EditCondition="bEnableMaterialValidation"))
 	TArray<FMaterialEditorValidationPlatform> MaterialValidationPlatforms;
 };
+
+#undef UE_API
diff --git a/Engine/Plugins/Editor/DataValidation/Source/DataValidation/Public/DirtyFilesChangelistValidator.h b/Engine/Plugins/Editor/DataValidation/Source/DataValidation/Public/DirtyFilesChangelistValidator.h
--- a/Engine/Plugins/Editor/DataValidation/Source/DataValidation/Public/DirtyFilesChangelistValidator.h
+++ b/Engine/Plugins/Editor/DataValidation/Source/DataValidation/Public/DirtyFilesChangelistValidator.h
@@ -5,17 +5,21 @@
 #include "EditorValidatorBase.h"
 #include "DirtyFilesChangelistValidator.generated.h"
 
+#define UE_API DATAVALIDATION_API
+
 class UPackage;
 
 /**
 * Validates there is no unsaved files in the changelist about to be submitted.
 */
-UCLASS()
-class DATAVALIDATION_API UDirtyFilesChangelistValidator : public UEditorValidatorBase
+UCLASS(MinimalAPI)
+class UDirtyFilesChangelistValidator : public UEditorValidatorBase
 {
 	GENERATED_BODY()
 
 protected:
-	virtual bool CanValidateAsset_Implementation(const FAssetData& InAssetData, UObject* InAsset, FDataValidationContext& InContext) const override;
-	virtual EDataValidationResult ValidateLoadedAsset_Implementation(const FAssetData& InAssetData, UObject* InAsset, FDataValidationContext& InContext) override;
+	UE_API virtual bool CanValidateAsset_Implementation(const FAssetData& InAssetData, UObject* InAsset, FDataValidationContext& InContext) const override;
+	UE_API virtual EDataValidationResult ValidateLoadedAsset_Implementation(const FAssetData& InAssetData, UObject* InAsset, FDataValidationContext& InContext) override;
 };
+
+#undef UE_API
diff --git a/Engine/Plugins/Editor/DataValidation/Source/DataValidation/Public/EditorValidatorBase.h b/Engine/Plugins/Editor/DataValidation/Source/DataValidation/Public/EditorValidatorBase.h
--- a/Engine/Plugins/Editor/DataValidation/Source/DataValidation/Public/EditorValidatorBase.h
+++ b/Engine/Plugins/Editor/DataValidation/Source/DataValidation/Public/EditorValidatorBase.h
@@ -7,6 +7,8 @@
 
 #include "EditorValidatorBase.generated.h"
 
+#define UE_API DATAVALIDATION_API
+
 enum class EDataValidationUsecase : uint8;
 namespace EMessageSeverity { enum Type : int; };
 struct FAssetData;
@@ -22,17 +24,17 @@ class UDataValidationChangelist;
 * C++ and Blueprint validators will be gathered on editor start, while python validators need
 * to register themselves
 */
-UCLASS(Abstract, Blueprintable, meta = (ShowWorldContextPin), Config=Editor)
-class DATAVALIDATION_API UEditorValidatorBase : public UObject
+UCLASS(MinimalAPI, Abstract, Blueprintable, meta = (ShowWorldContextPin), Config=Editor)
+class UEditorValidatorBase : public UObject
 {
 	GENERATED_BODY()
 
 public:
-	UEditorValidatorBase();
+	UE_API UEditorValidatorBase();
 	
 	// Validation entry point - will dispatch to native and/or BP validation as appropriate and return results.
 	// If validator doesn't want to validate this asset/in this context, it will return NotValidated which is acceptable.
-	EDataValidationResult ValidateLoadedAsset(const FAssetData& InAssetData, UObject* InAsset, FDataValidationContext& Context);
+	UE_API EDataValidationResult ValidateLoadedAsset(const FAssetData& InAssetData, UObject* InAsset, FDataValidationContext& Context);
 
 	/** 
 	 * Override this to determine whether or not you can use this validator given this context.
@@ -57,46 +59,46 @@ public:
 	
 	/** Override this to determine whether or not you can use this validator given this usecase */
 	UFUNCTION(BlueprintNativeEvent, Category = "Asset Validation", DisplayName="Can Validate")
-	bool K2_CanValidate(const EDataValidationUsecase InUsecase) const;
+	UE_API bool K2_CanValidate(const EDataValidationUsecase InUsecase) const;
 	
 	/** Override this to determine whether or not you can validate a given asset with this validator */
 	UFUNCTION(BlueprintNativeEvent, Category = "Asset Validation", DisplayName="Can Validate Asset")
-	bool K2_CanValidateAsset(UObject* InAsset) const;
+	UE_API bool K2_CanValidateAsset(UObject* InAsset) const;
 	
 	/** Override this in blueprint to validate assets */
 	UFUNCTION(BlueprintNativeEvent, Category = "Asset Validation", DisplayName="Validate Loaded Asset")
-	EDataValidationResult K2_ValidateLoadedAsset(UObject* InAsset);
+	UE_API EDataValidationResult K2_ValidateLoadedAsset(UObject* InAsset);
 
 	// Declaration for above BlueprintNativeEvent
-	virtual EDataValidationResult K2_ValidateLoadedAsset_Implementation(UObject* InAsset);
+	UE_API virtual EDataValidationResult K2_ValidateLoadedAsset_Implementation(UObject* InAsset);
 	
 	/** Marks the validation as failed and adds an error message. */
 	UFUNCTION(BlueprintCallable, Category = "Asset Validation")
-	void AssetFails(const UObject* InAsset, const FText& InMessage);
+	UE_API void AssetFails(const UObject* InAsset, const FText& InMessage);
 
 	UE_DEPRECATED("5.4", "This function has been replaced by the version AssetFails(UObject*, const FText&)")
-	void AssetFails(const UObject* InAsset, const FText& InMessage, TArray<FText>& InOutErrors);
+	UE_API void AssetFails(const UObject* InAsset, const FText& InMessage, TArray<FText>& InOutErrors);
 
 	/** Marks the validation as successful. Failure to call this will report the validator as not having checked the asset. */
 	UFUNCTION(BlueprintCallable, Category = "Asset Validation")
-	void AssetPasses(const UObject* InAsset);
+	UE_API void AssetPasses(const UObject* InAsset);
 
 	/** Adds a message to this validation but doesn't mark it as failed. */
 	UFUNCTION(BlueprintCallable, Category = "Asset Validation")
-	void AssetWarning(const UObject* InAsset, const FText& InMessage);
+	UE_API void AssetWarning(const UObject* InAsset, const FText& InMessage);
 	
 	/** 
 	 * Add a tokenized message to the validation results. If the severity is error, marks the validation as failed.
 	 */
-	TSharedRef<FTokenizedMessage> AssetMessage(const FAssetData& InAssetData, EMessageSeverity::Type InSeverity, const FText& InText = {});
-	TSharedRef<FTokenizedMessage> AssetMessage(EMessageSeverity::Type InSeverity, const FText& InText = {});
+	UE_API TSharedRef<FTokenizedMessage> AssetMessage(const FAssetData& InAssetData, EMessageSeverity::Type InSeverity, const FText& InText = {});
+	UE_API TSharedRef<FTokenizedMessage> AssetMessage(EMessageSeverity::Type InSeverity, const FText& InText = {});
 
 	virtual bool IsEnabled() const
 	{
 		return bIsEnabled && !bIsConfigDisabled;
 	}
 
-	void ResetValidationState();
+	UE_API void ResetValidationState();
 
 	bool IsValidationStateSet() const 
 	{
@@ -109,10 +111,10 @@ public:
 		return ValidationResult;
 	}
 
-	const TArray<FText>& GetAllWarnings() const;
+	UE_API const TArray<FText>& GetAllWarnings() const;
 	
-	void AddLegacyValidationErrors(TArray<FText> InErrors);
-	EDataValidationResult ExtractValidationState(FDataValidationContext& InOutContext) const;
+	UE_API void AddLegacyValidationErrors(TArray<FText> InErrors);
+	UE_API EDataValidationResult ExtractValidationState(FDataValidationContext& InOutContext) const;
 
 	// Legacy entry points from old interface 
 	UE_DEPRECATED("5.4", "CanValidate_Implementation(const EDataValidationUsecase InUsecase) is deprecated, override CanValidateAsset_Implementation(UObject* InObject, FDataValidationContext& InContext) instead")
@@ -155,3 +157,5 @@ private:
 
 	const FAssetData* CurrentAssetBeingValidated;
 };
+
+#undef UE_API
diff --git a/Engine/Plugins/Editor/DataValidation/Source/DataValidation/Public/EditorValidatorHelpers.h b/Engine/Plugins/Editor/DataValidation/Source/DataValidation/Public/EditorValidatorHelpers.h
--- a/Engine/Plugins/Editor/DataValidation/Source/DataValidation/Public/EditorValidatorHelpers.h
+++ b/Engine/Plugins/Editor/DataValidation/Source/DataValidation/Public/EditorValidatorHelpers.h
@@ -7,6 +7,8 @@
 #include "Templates/Tuple.h"
 #include "UObject/NameTypes.h"
 
+#define UE_API DATAVALIDATION_API
+
 class FMessageLog;
 class FText;
 struct FAssetData;
@@ -26,10 +28,10 @@ namespace UE::DataValidation
     // Utility to capture log messages from all threads during data validation so they can be filtered
     // and passed as validation errors if desired.
     struct FLogMessageGathererImpl;
-    struct DATAVALIDATION_API FScopedLogMessageGatherer
+    struct FScopedLogMessageGatherer
     {
-        FScopedLogMessageGatherer(bool bEnabled = true);
-        ~FScopedLogMessageGatherer();
+        UE_API FScopedLogMessageGatherer(bool bEnabled = true);
+        UE_API ~FScopedLogMessageGatherer();
 
         FScopedLogMessageGatherer(const FScopedLogMessageGatherer&) = delete;
         FScopedLogMessageGatherer& operator=(const FScopedLogMessageGatherer&) = delete;
@@ -38,48 +40,48 @@ namespace UE::DataValidation
         FScopedLogMessageGatherer& operator=(FScopedLogMessageGatherer&&) = delete;
 
         /** Returns the current gather constructed on this thread in an enclosing scope if any */
-        static FScopedLogMessageGatherer* GetCurrentThreadGatherer();
+        static UE_API FScopedLogMessageGatherer* GetCurrentThreadGatherer();
         
         /** 
          * Stop gathering logs and return the warnings and errors gathered so far
          */
-        void Stop(TArray<FString>& OutWarnings, TArray<FString>& OutErrors);
+        UE_API void Stop(TArray<FString>& OutWarnings, TArray<FString>& OutErrors);
 
         /** 
          * Tell this scope to ignore certain log messages 
          * @return A tuple that can be passed to RemoveIgnoreCategories to reverse this 
          */
-        TTuple<int32, int32> AddIgnoreCategories(TConstArrayView<FName> NewCategories);
+        UE_API TTuple<int32, int32> AddIgnoreCategories(TConstArrayView<FName> NewCategories);
 
         /** 
          * Tell this scope to stop ignoring certain log messages.
          * @param Range The return value from a previous call to AddIgnoreCategories 
          */
-        void RemoveIgnoreCategories(TTuple<int32, int32> Range);
+        UE_API void RemoveIgnoreCategories(TTuple<int32, int32> Range);
 
         /** 
          * Tell this scope to ignore certain log messages 
          * @return A tuple that can be passed to RemoveIgnorePatterns to reverse this 
          */
-        TTuple<int32, int32> AddIgnorePatterns(TConstArrayView<FWildcardString> NewPatterns);
+        UE_API TTuple<int32, int32> AddIgnorePatterns(TConstArrayView<FWildcardString> NewPatterns);
 
         /** 
          * Tell this scope to stop ignoring certain log messages.
          * @param Range The return value from a previous call to AddIgnorePatterns 
          */
-        void RemoveIgnorePatterns(TTuple<int32, int32> Range);
+        UE_API void RemoveIgnorePatterns(TTuple<int32, int32> Range);
         
         /** 
          * Delegate that may be bound by a project to set a default list of log categories to ignore 
          * during asset validation
          */
-        static TDelegate<TArray<FName>(void)>& GetDefaultIgnoreCategoriesDelegate();
+        static UE_API TDelegate<TArray<FName>(void)>& GetDefaultIgnoreCategoriesDelegate();
 
         /**
          * Delegate that may be bound by a project to set a default list of log patterns to ignore 
          * during asset validation
          */
-        static TDelegate<TArray<FWildcardString>(void)>& GetDefaultIgnorePatternsDelegate();
+        static UE_API TDelegate<TArray<FWildcardString>(void)>& GetDefaultIgnorePatternsDelegate();
         
     protected:
         FScopedLogMessageGatherer* Previous = nullptr;
@@ -87,10 +89,10 @@ namespace UE::DataValidation
     };
     
     // Utility to tell the closest enclosing FScopedLogMessageGatherer to ignore certain messages
-    struct DATAVALIDATION_API FScopedIgnoreLogMessages
+    struct FScopedIgnoreLogMessages
     {
-        FScopedIgnoreLogMessages(TConstArrayView<FWildcardString> Patterns, TConstArrayView<FName> Categories);
-        ~FScopedIgnoreLogMessages();
+        UE_API FScopedIgnoreLogMessages(TConstArrayView<FWildcardString> Patterns, TConstArrayView<FName> Categories);
+        UE_API ~FScopedIgnoreLogMessages();
         
     protected:
         TTuple<int32, int32> PatternsRange = {0, 0};
@@ -114,4 +116,6 @@ namespace UE::DataValidation
         FMessageLog& Log,
         EMessageSeverity::Type Severity,
         TConstArrayView<FText> Messages);
-}
\ No newline at end of file
+}
+
+#undef UE_API
diff --git a/Engine/Plugins/Editor/DataValidation/Source/DataValidation/Public/EditorValidator_Localization.h b/Engine/Plugins/Editor/DataValidation/Source/DataValidation/Public/EditorValidator_Localization.h
--- a/Engine/Plugins/Editor/DataValidation/Source/DataValidation/Public/EditorValidator_Localization.h
+++ b/Engine/Plugins/Editor/DataValidation/Source/DataValidation/Public/EditorValidator_Localization.h
@@ -5,23 +5,27 @@
 #include "EditorValidatorBase.h"
 #include "EditorValidator_Localization.generated.h"
 
+#define UE_API DATAVALIDATION_API
+
 /*
 * Validates that localized assets (within the L10N folder) conform to a corresponding source asset of the correct type.
 * Localized assets that fail this validation will never be loaded as localized variants at runtime.
 */
-UCLASS()
-class DATAVALIDATION_API UEditorValidator_Localization : public UEditorValidatorBase
+UCLASS(MinimalAPI)
+class UEditorValidator_Localization : public UEditorValidatorBase
 {
 	GENERATED_BODY()
 
 public:
-	UEditorValidator_Localization();
+	UE_API UEditorValidator_Localization();
 
 protected:
-	virtual bool CanValidateAsset_Implementation(const FAssetData& InAssetData, UObject* InAsset, FDataValidationContext& InContext) const override;
-	virtual EDataValidationResult ValidateLoadedAsset_Implementation(const FAssetData& InAssetData, UObject* InAsset, FDataValidationContext& InContext) override;
+	UE_API virtual bool CanValidateAsset_Implementation(const FAssetData& InAssetData, UObject* InAsset, FDataValidationContext& InContext) const override;
+	UE_API virtual EDataValidationResult ValidateLoadedAsset_Implementation(const FAssetData& InAssetData, UObject* InAsset, FDataValidationContext& InContext) override;
 
-	const TArray<FString>* FindOrCacheCulturesForLocalizedRoot(const FString& InLocalizedRootPath);
+	UE_API const TArray<FString>* FindOrCacheCulturesForLocalizedRoot(const FString& InLocalizedRootPath);
 
 	TMap<FString, TArray<FString>> CachedCulturesForLocalizedRoots;
 };
+
+#undef UE_API
diff --git a/Engine/Plugins/Editor/EngineAssetDefinitions/Source/Private/AssetDefinition_CanvasRenderTarget2D.h b/Engine/Plugins/Editor/EngineAssetDefinitions/Source/Private/AssetDefinition_CanvasRenderTarget2D.h
--- a/Engine/Plugins/Editor/EngineAssetDefinitions/Source/Private/AssetDefinition_CanvasRenderTarget2D.h
+++ b/Engine/Plugins/Editor/EngineAssetDefinitions/Source/Private/AssetDefinition_CanvasRenderTarget2D.h
@@ -7,8 +7,8 @@
 
 #include "AssetDefinition_CanvasRenderTarget2D.generated.h"
 
-UCLASS()
-class ENGINEASSETDEFINITIONS_API UAssetDefinition_CanvasRenderTarget2D : public UAssetDefinition_TextureRenderTarget2D
+UCLASS(MinimalAPI)
+class UAssetDefinition_CanvasRenderTarget2D : public UAssetDefinition_TextureRenderTarget2D
 {
 	GENERATED_BODY()
 
diff --git a/Engine/Plugins/Editor/EngineAssetDefinitions/Source/Private/AssetDefinition_Texture2D.h b/Engine/Plugins/Editor/EngineAssetDefinitions/Source/Private/AssetDefinition_Texture2D.h
--- a/Engine/Plugins/Editor/EngineAssetDefinitions/Source/Private/AssetDefinition_Texture2D.h
+++ b/Engine/Plugins/Editor/EngineAssetDefinitions/Source/Private/AssetDefinition_Texture2D.h
@@ -7,8 +7,8 @@
 
 #include "AssetDefinition_Texture2D.generated.h"
 
-UCLASS()
-class ENGINEASSETDEFINITIONS_API UAssetDefinition_Texture2D : public UAssetDefinition_Texture
+UCLASS(MinimalAPI)
+class UAssetDefinition_Texture2D : public UAssetDefinition_Texture
 {
 	GENERATED_BODY()
 
diff --git a/Engine/Plugins/Editor/EngineAssetDefinitions/Source/Private/AssetDefinition_Texture2DArray.h b/Engine/Plugins/Editor/EngineAssetDefinitions/Source/Private/AssetDefinition_Texture2DArray.h
--- a/Engine/Plugins/Editor/EngineAssetDefinitions/Source/Private/AssetDefinition_Texture2DArray.h
+++ b/Engine/Plugins/Editor/EngineAssetDefinitions/Source/Private/AssetDefinition_Texture2DArray.h
@@ -7,8 +7,8 @@
 
 #include "AssetDefinition_Texture2DArray.generated.h"
 
-UCLASS()
-class ENGINEASSETDEFINITIONS_API UAssetDefinition_Texture2DArray : public UAssetDefinition_Texture
+UCLASS(MinimalAPI)
+class UAssetDefinition_Texture2DArray : public UAssetDefinition_Texture
 {
 	GENERATED_BODY()
 
diff --git a/Engine/Plugins/Editor/EngineAssetDefinitions/Source/Private/AssetDefinition_TextureCube.h b/Engine/Plugins/Editor/EngineAssetDefinitions/Source/Private/AssetDefinition_TextureCube.h
--- a/Engine/Plugins/Editor/EngineAssetDefinitions/Source/Private/AssetDefinition_TextureCube.h
+++ b/Engine/Plugins/Editor/EngineAssetDefinitions/Source/Private/AssetDefinition_TextureCube.h
@@ -7,8 +7,8 @@
 
 #include "AssetDefinition_TextureCube.generated.h"
 
-UCLASS()
-class ENGINEASSETDEFINITIONS_API UAssetDefinition_TextureCube : public UAssetDefinition_Texture
+UCLASS(MinimalAPI)
+class UAssetDefinition_TextureCube : public UAssetDefinition_Texture
 {
 	GENERATED_BODY()
 
diff --git a/Engine/Plugins/Editor/EngineAssetDefinitions/Source/Private/AssetDefinition_TextureCubeArray.h b/Engine/Plugins/Editor/EngineAssetDefinitions/Source/Private/AssetDefinition_TextureCubeArray.h
--- a/Engine/Plugins/Editor/EngineAssetDefinitions/Source/Private/AssetDefinition_TextureCubeArray.h
+++ b/Engine/Plugins/Editor/EngineAssetDefinitions/Source/Private/AssetDefinition_TextureCubeArray.h
@@ -7,8 +7,8 @@
 
 #include "AssetDefinition_TextureCubeArray.generated.h"
 
-UCLASS()
-class ENGINEASSETDEFINITIONS_API UAssetDefinition_TextureCubeArray : public UAssetDefinition_Texture
+UCLASS(MinimalAPI)
+class UAssetDefinition_TextureCubeArray : public UAssetDefinition_Texture
 {
 	GENERATED_BODY()
 
diff --git a/Engine/Plugins/Editor/EngineAssetDefinitions/Source/Private/AssetDefinition_TextureLightProfile.h b/Engine/Plugins/Editor/EngineAssetDefinitions/Source/Private/AssetDefinition_TextureLightProfile.h
--- a/Engine/Plugins/Editor/EngineAssetDefinitions/Source/Private/AssetDefinition_TextureLightProfile.h
+++ b/Engine/Plugins/Editor/EngineAssetDefinitions/Source/Private/AssetDefinition_TextureLightProfile.h
@@ -7,8 +7,8 @@
 
 #include "AssetDefinition_TextureLightProfile.generated.h"
 
-UCLASS()
-class ENGINEASSETDEFINITIONS_API UAssetDefinition_TextureLightProfile : public UAssetDefinition_Texture
+UCLASS(MinimalAPI)
+class UAssetDefinition_TextureLightProfile : public UAssetDefinition_Texture
 {
 	GENERATED_BODY()
 
diff --git a/Engine/Plugins/Editor/EngineAssetDefinitions/Source/Private/AssetDefinition_TextureRenderTarget.h b/Engine/Plugins/Editor/EngineAssetDefinitions/Source/Private/AssetDefinition_TextureRenderTarget.h
--- a/Engine/Plugins/Editor/EngineAssetDefinitions/Source/Private/AssetDefinition_TextureRenderTarget.h
+++ b/Engine/Plugins/Editor/EngineAssetDefinitions/Source/Private/AssetDefinition_TextureRenderTarget.h
@@ -7,8 +7,8 @@
 
 #include "AssetDefinition_TextureRenderTarget.generated.h"
 
-UCLASS()
-class ENGINEASSETDEFINITIONS_API UAssetDefinition_TextureRenderTarget : public UAssetDefinition_Texture
+UCLASS(MinimalAPI)
+class UAssetDefinition_TextureRenderTarget : public UAssetDefinition_Texture
 {
 	GENERATED_BODY()
 
diff --git a/Engine/Plugins/Editor/EngineAssetDefinitions/Source/Private/AssetDefinition_TextureRenderTarget2D.h b/Engine/Plugins/Editor/EngineAssetDefinitions/Source/Private/AssetDefinition_TextureRenderTarget2D.h
--- a/Engine/Plugins/Editor/EngineAssetDefinitions/Source/Private/AssetDefinition_TextureRenderTarget2D.h
+++ b/Engine/Plugins/Editor/EngineAssetDefinitions/Source/Private/AssetDefinition_TextureRenderTarget2D.h
@@ -7,8 +7,8 @@
 
 #include "AssetDefinition_TextureRenderTarget2D.generated.h"
 
-UCLASS()
-class ENGINEASSETDEFINITIONS_API UAssetDefinition_TextureRenderTarget2D : public UAssetDefinition_TextureRenderTarget
+UCLASS(MinimalAPI)
+class UAssetDefinition_TextureRenderTarget2D : public UAssetDefinition_TextureRenderTarget
 {
 	GENERATED_BODY()
 
diff --git a/Engine/Plugins/Editor/EngineAssetDefinitions/Source/Private/AssetDefinition_TextureRenderTarget2DArray.h b/Engine/Plugins/Editor/EngineAssetDefinitions/Source/Private/AssetDefinition_TextureRenderTarget2DArray.h
--- a/Engine/Plugins/Editor/EngineAssetDefinitions/Source/Private/AssetDefinition_TextureRenderTarget2DArray.h
+++ b/Engine/Plugins/Editor/EngineAssetDefinitions/Source/Private/AssetDefinition_TextureRenderTarget2DArray.h
@@ -7,8 +7,8 @@
 
 #include "AssetDefinition_TextureRenderTarget2DArray.generated.h"
 
-UCLASS()
-class ENGINEASSETDEFINITIONS_API UAssetDefinition_TextureRenderTarget2DArray : public UAssetDefinition_TextureRenderTarget
+UCLASS(MinimalAPI)
+class UAssetDefinition_TextureRenderTarget2DArray : public UAssetDefinition_TextureRenderTarget
 {
 	GENERATED_BODY()
 
diff --git a/Engine/Plugins/Editor/EngineAssetDefinitions/Source/Private/AssetDefinition_TextureRenderTargetCube.h b/Engine/Plugins/Editor/EngineAssetDefinitions/Source/Private/AssetDefinition_TextureRenderTargetCube.h
--- a/Engine/Plugins/Editor/EngineAssetDefinitions/Source/Private/AssetDefinition_TextureRenderTargetCube.h
+++ b/Engine/Plugins/Editor/EngineAssetDefinitions/Source/Private/AssetDefinition_TextureRenderTargetCube.h
@@ -7,8 +7,8 @@
 
 #include "AssetDefinition_TextureRenderTargetCube.generated.h"
 
-UCLASS()
-class ENGINEASSETDEFINITIONS_API UAssetDefinition_TextureRenderTargetCube : public UAssetDefinition_TextureRenderTarget
+UCLASS(MinimalAPI)
+class UAssetDefinition_TextureRenderTargetCube : public UAssetDefinition_TextureRenderTarget
 {
 	GENERATED_BODY()
 
diff --git a/Engine/Plugins/Editor/EngineAssetDefinitions/Source/Private/AssetDefinition_TextureRenderTargetVolume.h b/Engine/Plugins/Editor/EngineAssetDefinitions/Source/Private/AssetDefinition_TextureRenderTargetVolume.h
--- a/Engine/Plugins/Editor/EngineAssetDefinitions/Source/Private/AssetDefinition_TextureRenderTargetVolume.h
+++ b/Engine/Plugins/Editor/EngineAssetDefinitions/Source/Private/AssetDefinition_TextureRenderTargetVolume.h
@@ -7,8 +7,8 @@
 
 #include "AssetDefinition_TextureRenderTargetVolume.generated.h"
 
-UCLASS()
-class ENGINEASSETDEFINITIONS_API UAssetDefinition_TextureRenderTargetVolume : public UAssetDefinition_TextureRenderTarget
+UCLASS(MinimalAPI)
+class UAssetDefinition_TextureRenderTargetVolume : public UAssetDefinition_TextureRenderTarget
 {
 	GENERATED_BODY()
 
diff --git a/Engine/Plugins/Editor/EngineAssetDefinitions/Source/Private/AssetDefinition_VolumeTexture.h b/Engine/Plugins/Editor/EngineAssetDefinitions/Source/Private/AssetDefinition_VolumeTexture.h
--- a/Engine/Plugins/Editor/EngineAssetDefinitions/Source/Private/AssetDefinition_VolumeTexture.h
+++ b/Engine/Plugins/Editor/EngineAssetDefinitions/Source/Private/AssetDefinition_VolumeTexture.h
@@ -7,8 +7,8 @@
 
 #include "AssetDefinition_VolumeTexture.generated.h"
 
-UCLASS()
-class ENGINEASSETDEFINITIONS_API UAssetDefinition_VolumeTexture : public UAssetDefinition_Texture
+UCLASS(MinimalAPI)
+class UAssetDefinition_VolumeTexture : public UAssetDefinition_Texture
 {
 	GENERATED_BODY()
 
diff --git a/Engine/Plugins/Editor/EngineAssetDefinitions/Source/Public/AssetDefinition_DataAsset.h b/Engine/Plugins/Editor/EngineAssetDefinitions/Source/Public/AssetDefinition_DataAsset.h
--- a/Engine/Plugins/Editor/EngineAssetDefinitions/Source/Public/AssetDefinition_DataAsset.h
+++ b/Engine/Plugins/Editor/EngineAssetDefinitions/Source/Public/AssetDefinition_DataAsset.h
@@ -7,8 +7,10 @@
 
 #include "AssetDefinition_DataAsset.generated.h"
 
-UCLASS()
-class ENGINEASSETDEFINITIONS_API UAssetDefinition_DataAsset : public UAssetDefinitionDefault
+#define UE_API ENGINEASSETDEFINITIONS_API
+
+UCLASS(MinimalAPI)
+class UAssetDefinition_DataAsset : public UAssetDefinitionDefault
 {
 	GENERATED_BODY()
 
@@ -19,13 +21,15 @@ public:
 		IncludeClassInFilter = EIncludeClassInFilter::Always;
 	}
 	virtual FText GetAssetDisplayName() const override { return NSLOCTEXT("AssetDefinition", "DataAsset", "Data Asset"); }
-	virtual FText GetAssetDisplayName(const FAssetData& AssetData) const override;
+	UE_API virtual FText GetAssetDisplayName(const FAssetData& AssetData) const override;
 	virtual FLinearColor GetAssetColor() const override { return FLinearColor(FColor(201, 29, 85)); }
 	virtual TSoftClassPtr<UObject> GetAssetClass() const override { return UDataAsset::StaticClass(); }
 	
-	virtual EAssetCommandResult PerformAssetDiff(const FAssetDiffArgs& DiffArgs) const override;
+	UE_API virtual EAssetCommandResult PerformAssetDiff(const FAssetDiffArgs& DiffArgs) const override;
 
-	virtual bool CanMerge() const override;
-	virtual EAssetCommandResult Merge(const FAssetAutomaticMergeArgs& MergeArgs) const override;
-	virtual EAssetCommandResult Merge(const FAssetManualMergeArgs& MergeArgs) const override;
+	UE_API virtual bool CanMerge() const override;
+	UE_API virtual EAssetCommandResult Merge(const FAssetAutomaticMergeArgs& MergeArgs) const override;
+	UE_API virtual EAssetCommandResult Merge(const FAssetManualMergeArgs& MergeArgs) const override;
 };
+
+#undef UE_API
diff --git a/Engine/Plugins/Editor/EngineAssetDefinitions/Source/Public/AssetDefinition_Redirector.h b/Engine/Plugins/Editor/EngineAssetDefinitions/Source/Public/AssetDefinition_Redirector.h
--- a/Engine/Plugins/Editor/EngineAssetDefinitions/Source/Public/AssetDefinition_Redirector.h
+++ b/Engine/Plugins/Editor/EngineAssetDefinitions/Source/Public/AssetDefinition_Redirector.h
@@ -6,11 +6,13 @@
 
 #include "AssetDefinition_Redirector.generated.h"
 
+#define UE_API ENGINEASSETDEFINITIONS_API
+
 enum class EAssetCommandResult : uint8;
 struct FAssetActivateArgs;
 
-UCLASS()
-class ENGINEASSETDEFINITIONS_API UAssetDefinition_Redirector : public UAssetDefinitionDefault
+UCLASS(MinimalAPI)
+class UAssetDefinition_Redirector : public UAssetDefinitionDefault
 {
 	GENERATED_BODY()
 
@@ -19,5 +21,7 @@ public:
 	virtual FText GetAssetDisplayName() const override { return NSLOCTEXT("AssetDefinition", "Redirector", "Redirector"); }
 	virtual FLinearColor GetAssetColor() const override { return FLinearColor(FColor(128, 128, 128)); }
 	virtual TSoftClassPtr<UObject> GetAssetClass() const override { return UObjectRedirector::StaticClass(); }
-	virtual EAssetCommandResult ActivateAssets(const FAssetActivateArgs& ActivateArgs) const override;
+	UE_API virtual EAssetCommandResult ActivateAssets(const FAssetActivateArgs& ActivateArgs) const override;
 };
+
+#undef UE_API
diff --git a/Engine/Plugins/Editor/EngineAssetDefinitions/Source/Public/AssetDefinition_World.h b/Engine/Plugins/Editor/EngineAssetDefinitions/Source/Public/AssetDefinition_World.h
--- a/Engine/Plugins/Editor/EngineAssetDefinitions/Source/Public/AssetDefinition_World.h
+++ b/Engine/Plugins/Editor/EngineAssetDefinitions/Source/Public/AssetDefinition_World.h
@@ -8,14 +8,16 @@
 #include "Styling/AppStyle.h"
 #include "AssetDefinition_World.generated.h"
 
+#define UE_API ENGINEASSETDEFINITIONS_API
+
 enum class EAssetCommandResult : uint8;
 struct FAssetActivateArgs;
 struct FAssetCategoryPath;
 struct FAssetOpenArgs;
 struct FAssetSupportResponse;
 
-UCLASS()
-class ENGINEASSETDEFINITIONS_API UAssetDefinition_World : public UAssetDefinitionDefault
+UCLASS(MinimalAPI)
+class UAssetDefinition_World : public UAssetDefinitionDefault
 {
 	GENERATED_BODY()
 
@@ -24,18 +26,20 @@ public:
 	virtual FText GetAssetDisplayName() const override { return NSLOCTEXT("AssetTypeActions", "AssetTypeActions_World", "Level"); }
 	virtual FLinearColor GetAssetColor() const override { return FAppStyle::Get().GetColor("LevelEditor.AssetColor"); }
 	virtual TSoftClassPtr<UObject> GetAssetClass() const override { return UWorld::StaticClass(); }
-	virtual TConstArrayView<FAssetCategoryPath> GetAssetCategories() const override;
+	UE_API virtual TConstArrayView<FAssetCategoryPath> GetAssetCategories() const override;
 
-	virtual TArray<FAssetData> PrepareToActivateAssets(const FAssetActivateArgs& ActivateArgs) const override;
-	virtual EAssetCommandResult OpenAssets(const FAssetOpenArgs& OpenArgs) const override;
-	virtual EAssetCommandResult PerformAssetDiff(const FAssetDiffArgs& DiffArgs) const override;
+	UE_API virtual TArray<FAssetData> PrepareToActivateAssets(const FAssetActivateArgs& ActivateArgs) const override;
+	UE_API virtual EAssetCommandResult OpenAssets(const FAssetOpenArgs& OpenArgs) const override;
+	UE_API virtual EAssetCommandResult PerformAssetDiff(const FAssetDiffArgs& DiffArgs) const override;
 	
-	virtual FAssetSupportResponse CanRename(const FAssetData& InAsset) const override;
-	virtual FAssetSupportResponse CanDuplicate(const FAssetData& InAsset) const override;
+	UE_API virtual FAssetSupportResponse CanRename(const FAssetData& InAsset) const override;
+	UE_API virtual FAssetSupportResponse CanDuplicate(const FAssetData& InAsset) const override;
 
-	virtual UThumbnailInfo* LoadThumbnailInfo(const FAssetData& InAsset) const override;
+	UE_API virtual UThumbnailInfo* LoadThumbnailInfo(const FAssetData& InAsset) const override;
 	// UAssetDefinition End
 	
 public:
-	bool IsPartitionWorldInUse(const FAssetData& InAsset) const;
+	UE_API bool IsPartitionWorldInUse(const FAssetData& InAsset) const;
 };
+
+#undef UE_API
diff --git a/Engine/Plugins/Editor/EngineAssetDefinitions/Source/Public/Table/AssetDefinition_DataTable.h b/Engine/Plugins/Editor/EngineAssetDefinitions/Source/Public/Table/AssetDefinition_DataTable.h
--- a/Engine/Plugins/Editor/EngineAssetDefinitions/Source/Public/Table/AssetDefinition_DataTable.h
+++ b/Engine/Plugins/Editor/EngineAssetDefinitions/Source/Public/Table/AssetDefinition_DataTable.h
@@ -7,15 +7,17 @@
 
 #include "AssetDefinition_DataTable.generated.h"
 
-UCLASS()
-class ENGINEASSETDEFINITIONS_API UAssetDefinition_DataTable : public UAssetDefinitionDefault
+#define UE_API ENGINEASSETDEFINITIONS_API
+
+UCLASS(MinimalAPI)
+class UAssetDefinition_DataTable : public UAssetDefinitionDefault
 {
 	GENERATED_BODY()
 
 public:
 	// UAssetDefinition Begin
 	virtual FText GetAssetDisplayName() const override { return NSLOCTEXT("AssetTypeActions", "AssetTypeActions_DataTable", "Data Table"); }
-	virtual FText GetAssetDisplayName(const FAssetData& AssetData) const override;
+	UE_API virtual FText GetAssetDisplayName(const FAssetData& AssetData) const override;
 	virtual FLinearColor GetAssetColor() const override { return FLinearColor(FColor(62, 140, 35)); }
 	virtual TSoftClassPtr<UObject> GetAssetClass() const override { return UDataTable::StaticClass(); }
 	virtual TConstArrayView<FAssetCategoryPath> GetAssetCategories() const override
@@ -24,7 +26,9 @@ public:
 		return Categories;
 	}
 	virtual bool CanImport() const override { return true; }
-	virtual EAssetCommandResult PerformAssetDiff(const FAssetDiffArgs& DiffArgs) const override;
-	virtual EAssetCommandResult OpenAssets(const FAssetOpenArgs& OpenArgs) const override;
+	UE_API virtual EAssetCommandResult PerformAssetDiff(const FAssetDiffArgs& DiffArgs) const override;
+	UE_API virtual EAssetCommandResult OpenAssets(const FAssetOpenArgs& OpenArgs) const override;
 	// UAssetDefinition End
 };
+
+#undef UE_API
