SourceControl: The refresh FUpdateStatus operation is a forced operation.

#rnx

diff --git a/Engine/Plugins/Editor/AssetManagerEditor/Source/AssetManagerEditor/Private/AssetSourceControlContextMenu.cpp b/Engine/Plugins/Editor/AssetManagerEditor/Source/AssetManagerEditor/Private/AssetSourceControlContextMenu.cpp
--- a/Engine/Plugins/Editor/AssetManagerEditor/Source/AssetManagerEditor/Private/AssetSourceControlContextMenu.cpp
+++ b/Engine/Plugins/Editor/AssetManagerEditor/Source/AssetManagerEditor/Private/AssetSourceControlContextMenu.cpp
@@ -1063,7 +1063,10 @@ void FAssetSourceControlContextMenuState::ExecuteSCCRefresh() const
 	TArray<FString> PackageNames;
 	GetSelectedPackageNames(PackageNames);
 
-	ISourceControlModule::Get().GetProvider().Execute(ISourceControlOperation::Create<FUpdateStatus>(), SourceControlHelpers::PackageFilenames(PackageNames), EConcurrency::Asynchronous);
+	TSharedPtr<FUpdateStatus, ESPMode::ThreadSafe> Operation = ISourceControlOperation::Create<FUpdateStatus>();
+	Operation->SetForceUpdate(true);
+
+	ISourceControlModule::Get().GetProvider().Execute(Operation.ToSharedRef(), SourceControlHelpers::PackageFilenames(PackageNames), EConcurrency::Asynchronous);
 }
 
 void FAssetSourceControlContextMenuState::ExecuteSCCCopyPaths() const
diff --git a/Engine/Plugins/Editor/ContentBrowser/ContentBrowserFileDataSource/Source/ContentBrowserFileDataSource/Private/FileSourceControlContextMenu.cpp b/Engine/Plugins/Editor/ContentBrowser/ContentBrowserFileDataSource/Source/ContentBrowserFileDataSource/Private/FileSourceControlContextMenu.cpp
--- a/Engine/Plugins/Editor/ContentBrowser/ContentBrowserFileDataSource/Source/ContentBrowserFileDataSource/Private/FileSourceControlContextMenu.cpp
+++ b/Engine/Plugins/Editor/ContentBrowser/ContentBrowserFileDataSource/Source/ContentBrowserFileDataSource/Private/FileSourceControlContextMenu.cpp
@@ -621,7 +621,10 @@ void FFileSourceControlContextMenuState::ExecuteSCCSync() const
 
 void FFileSourceControlContextMenuState::ExecuteSCCRefresh() const
 {
-	ISourceControlModule::Get().GetProvider().Execute(ISourceControlOperation::Create<FUpdateStatus>(), SelectedFiles, EConcurrency::Asynchronous);
+	TSharedPtr<FUpdateStatus, ESPMode::ThreadSafe> Operation = ISourceControlOperation::Create<FUpdateStatus>();
+	Operation->SetForceUpdate(true);
+
+	ISourceControlModule::Get().GetProvider().Execute(Operation.ToSharedRef(), SelectedFiles, EConcurrency::Asynchronous);
 }
 
 bool FFileSourceControlContextMenuState::CanExecuteSCCCheckOut() const
diff --git a/Engine/Source/Editor/SceneOutliner/Private/SceneOutlinerSCCHandler.cpp b/Engine/Source/Editor/SceneOutliner/Private/SceneOutlinerSCCHandler.cpp
--- a/Engine/Source/Editor/SceneOutliner/Private/SceneOutlinerSCCHandler.cpp
+++ b/Engine/Source/Editor/SceneOutliner/Private/SceneOutlinerSCCHandler.cpp
@@ -382,7 +382,10 @@ void FSceneOutlinerSCCHandler::ExecuteSCCRefresh()
 	TArray<FString> PackageNames;
 	GetSelectedPackageNames(PackageNames);
 
-	ISourceControlModule::Get().GetProvider().Execute(ISourceControlOperation::Create<FUpdateStatus>(), SourceControlHelpers::PackageFilenames(PackageNames), EConcurrency::Asynchronous);
+	TSharedPtr<FUpdateStatus, ESPMode::ThreadSafe> Operation = ISourceControlOperation::Create<FUpdateStatus>();
+	Operation->SetForceUpdate(true);
+
+	ISourceControlModule::Get().GetProvider().Execute(Operation.ToSharedRef(), SourceControlHelpers::PackageFilenames(PackageNames), EConcurrency::Asynchronous);
 }
 
 void FSceneOutlinerSCCHandler::ExecuteSCCCheckOut()
