Fixed type in UContentBrowserFileDataSource::EnumerateItemsAtUserProvidedPath.

#rnx

diff --git a/Engine/Plugins/Editor/ContentBrowser/ContentBrowserFileDataSource/Source/ContentBrowserFileDataSource/Private/ContentBrowserFileDataSource.cpp b/Engine/Plugins/Editor/ContentBrowser/ContentBrowserFileDataSource/Source/ContentBrowserFileDataSource/Private/ContentBrowserFileDataSource.cpp
--- a/Engine/Plugins/Editor/ContentBrowser/ContentBrowserFileDataSource/Source/ContentBrowserFileDataSource/Private/ContentBrowserFileDataSource.cpp
+++ b/Engine/Plugins/Editor/ContentBrowser/ContentBrowserFileDataSource/Source/ContentBrowserFileDataSource/Private/ContentBrowserFileDataSource.cpp
@@ -1011,7 +1011,7 @@ void UContentBrowserFileDataSource::EnumerateItemsAtUserProvidedPath(const FStri
 					case FDiscoveredItem::EType::File:
 						if (EnumHasAnyFlags(InItemTypeFilter, EContentBrowserItemTypeFilter::IncludeFiles))
 						{
-							InCallback(CreateFolderItem(InternalPath, DiscoveredItem->DiskPath));
+							InCallback(CreateFileItem(InternalPath, DiscoveredItem->DiskPath));
 						}
 						break;
 					}
